Sequencer: Renamed MovieSceneActor to LevelSequenceActor

diff --git a/Engine/Plugins/MovieScene/LevelSequenceEditor/Source/LevelSequenceEditor/Private/LevelSequenceEditorModule.cpp b/Engine/Plugins/MovieScene/LevelSequenceEditor/Source/LevelSequenceEditor/Private/LevelSequenceEditorModule.cpp
--- a/Engine/Plugins/MovieScene/LevelSequenceEditor/Source/LevelSequenceEditor/Private/LevelSequenceEditorModule.cpp
+++ b/Engine/Plugins/MovieScene/LevelSequenceEditor/Source/LevelSequenceEditor/Private/LevelSequenceEditorModule.cpp
@@ -2,9 +2,9 @@
 
 #include "LevelSequenceEditorPCH.h"
 #include "LevelEditor.h"
+#include "LevelSequenceActor.h"
 #include "LevelSequenceEditorStyle.h"
 #include "ModuleInterface.h"
-#include "MovieSceneActor.h"
 #include "PropertyEditorModule.h"
 
 #define LOCTEXT_NAMESPACE "LevelSequenceEditor"
@@ -175,13 +175,13 @@ protected:
 		}
 
 		// Spawn a  actor at the origin, and either move infront of the camera or focus camera on it (depending on the viewport) and open for edit
-		UActorFactory* ActorFactory = GEditor->FindActorFactoryForActorClass( AMovieSceneActor::StaticClass() );
+		UActorFactory* ActorFactory = GEditor->FindActorFactoryForActorClass( ALevelSequenceActor::StaticClass() );
 		if (!ensure(ActorFactory))
 		{
 			return;
 		}
 
-		AMovieSceneActor* NewActor = CastChecked<AMovieSceneActor>(GEditor->UseActorFactory(ActorFactory, FAssetData(NewAsset), &FTransform::Identity));
+		ALevelSequenceActor* NewActor = CastChecked<ALevelSequenceActor>(GEditor->UseActorFactory(ActorFactory, FAssetData(NewAsset), &FTransform::Identity));
 		if( GCurrentLevelEditingViewportClient->IsPerspective() )
 		{
 			GEditor->MoveActorInFrontOfCamera( *NewActor, GCurrentLevelEditingViewportClient->GetViewLocation(), GCurrentLevelEditingViewportClient->GetViewRotation().Vector() );
diff --git a/Engine/Source/Editor/EditorStyle/Private/SlateEditorStyle.cpp b/Engine/Source/Editor/EditorStyle/Private/SlateEditorStyle.cpp
--- a/Engine/Source/Editor/EditorStyle/Private/SlateEditorStyle.cpp
+++ b/Engine/Source/Editor/EditorStyle/Private/SlateEditorStyle.cpp
@@ -5355,7 +5355,7 @@ void FSlateEditorStyle::FStyle::SetupClassIconsAndThumbnails()
 			TEXT("MediaPlayer"),
 			TEXT("MediaSoundWave"),
 			TEXT("MediaTexture"),
-			TEXT("MovieSceneActor"),
+			TEXT("LevelSequenceActor"),
 			TEXT("MultiFont"),
 			TEXT("NavLinkProxy"),
 			TEXT("NavMeshBoundsVolume"),
diff --git a/Engine/Source/Editor/LevelEditor/Private/LevelEditorToolBar.cpp b/Engine/Source/Editor/LevelEditor/Private/LevelEditorToolBar.cpp
--- a/Engine/Source/Editor/LevelEditor/Private/LevelEditorToolBar.cpp
+++ b/Engine/Source/Editor/LevelEditor/Private/LevelEditorToolBar.cpp
@@ -15,7 +15,7 @@
 #include "Editor/ContentBrowser/Public/ContentBrowserModule.h"
 #include "EngineBuildSettings.h"
 #include "Matinee/MatineeActor.h"
-#include "MovieSceneActor.h"
+#include "LevelSequenceActor.h"
 #include "Engine/LevelScriptBlueprint.h"
 #include "ISettingsCategory.h"
 #include "ISettingsContainer.h"
@@ -2031,7 +2031,7 @@ TSharedRef< SWidget > FLevelEditorToolBar::GenerateCinematicsMenuContent( TShare
 	const bool bShouldCloseWindowAfterMenuSelection = true;
 	FMenuBuilder MenuBuilder( bShouldCloseWindowAfterMenuSelection, InCommandList, Extender );
 
-	// We can't build a list of Matinees and MovieSceneActors while the current World is a PIE world.
+	// We can't build a list of Matinees and LevelSequenceActors while the current World is a PIE world.
 	SceneOutliner::FInitializationOptions InitOptions;
 	{
 		InitOptions.Mode = ESceneOutlinerMode::ActorPicker;
@@ -2044,7 +2044,7 @@ TSharedRef< SWidget > FLevelEditorToolBar::GenerateCinematicsMenuContent( TShare
 
 		// Only display Matinee and MovieScene actors
 		auto ActorFilter = [](const AActor* Actor){
-			return Actor->IsA( AMatineeActor::StaticClass() ) || Actor->IsA( AMovieSceneActor::StaticClass() );
+			return Actor->IsA( AMatineeActor::StaticClass() ) || Actor->IsA( ALevelSequenceActor::StaticClass() );
 		};
 		InitOptions.Filters->AddFilterPredicate( SceneOutliner::FActorFilterPredicate::CreateLambda( ActorFilter ) );
 	}
@@ -2100,7 +2100,7 @@ TSharedRef< SWidget > FLevelEditorToolBar::GenerateCinematicsMenuContent( TShare
 	MenuBuilder.EndSection();
 
 	UWorld* World = LevelEditorWeakPtr.Pin()->GetWorld();
-	const bool bHasAnyMatineeActors = !!TActorIterator<AMatineeActor>(World) || !!TActorIterator<AMovieSceneActor>(World);
+	const bool bHasAnyMatineeActors = !!TActorIterator<AMatineeActor>(World) || !!TActorIterator<ALevelSequenceActor>(World);
 
 	//Add a heading to separate the existing matinees from the 'Add New Matinee Actor' button
 	MenuBuilder.BeginSection("LevelEditorExistingMatinee", LOCTEXT( "MatineeMenuCombo_ExistingHeading", "Edit Existing Matinee" ) );
@@ -2128,9 +2128,9 @@ void FLevelEditorToolBar::OnCinematicsActorPicked( AActor* Actor )
 		// Open Matinee for editing!
 		GEditor->OpenMatinee( MatineeActor );
 	}
-	else if (AMovieSceneActor* MovieSceneActor = Cast<AMovieSceneActor>(Actor))
+	else if (ALevelSequenceActor* LevelSequenceActor = Cast<ALevelSequenceActor>(Actor))
 	{
-		UObject* Asset = MovieSceneActor->LevelSequence.TryLoad();
+		UObject* Asset = LevelSequenceActor->LevelSequence.TryLoad();
 
 		if (Asset != nullptr)
 		{
diff --git a/Engine/Source/Editor/UnrealEd/Private/Factories/ActorFactory.cpp b/Engine/Source/Editor/UnrealEd/Private/Factories/ActorFactory.cpp
--- a/Engine/Source/Editor/UnrealEd/Private/Factories/ActorFactory.cpp
+++ b/Engine/Source/Editor/UnrealEd/Private/Factories/ActorFactory.cpp
@@ -68,7 +68,7 @@ ActorFactory.cpp:
 #include "Components/BillboardComponent.h"
 
 #include "LevelSequence.h"
-#include "MovieSceneActor.h"
+#include "LevelSequenceActor.h"
 #include "ActorFactoryMovieScene.h"
 
 DEFINE_LOG_CATEGORY_STATIC(LogActorFactory, Log, All);
@@ -1950,7 +1950,7 @@ UActorFactoryMovieScene::UActorFactoryMovieScene(const FObjectInitializer& Objec
 	: Super(ObjectInitializer)
 {
 	DisplayName = LOCTEXT("MovieSceneDisplayName", "MovieScene");
-	NewActorClass = AMovieSceneActor::StaticClass();
+	NewActorClass = ALevelSequenceActor::StaticClass();
 }
 
 bool UActorFactoryMovieScene::CanCreateActorFrom( const FAssetData& AssetData, FText& OutErrorMsg )
@@ -1971,13 +1971,13 @@ bool UActorFactoryMovieScene::CanCreateActorFrom( const FAssetData& AssetData, F
 
 AActor* UActorFactoryMovieScene::SpawnActor( UObject* Asset, ULevel* InLevel, const FVector& Location, const FRotator& Rotation, EObjectFlags ObjectFlags, const FName& Name )
 {
-	AMovieSceneActor* NewActor = Cast<AMovieSceneActor>(Super::SpawnActor(Asset, InLevel, Location, Rotation, ObjectFlags, Name));
+	ALevelSequenceActor* NewActor = Cast<ALevelSequenceActor>(Super::SpawnActor(Asset, InLevel, Location, Rotation, ObjectFlags, Name));
 
 	if (NewActor)
 	{
 		if (ULevelSequence* LevelSequence = Cast<ULevelSequence>(Asset))
 		{
-			NewActor->SetAnimation(LevelSequence);
+			NewActor->SetSequence(LevelSequence);
 		}
 	}
 
@@ -1986,9 +1986,9 @@ AActor* UActorFactoryMovieScene::SpawnActor( UObject* Asset, ULevel* InLevel, co
 
 UObject* UActorFactoryMovieScene::GetAssetFromActorInstance(AActor* Instance)
 {
-	if (AMovieSceneActor* MovieSceneActor = Cast<AMovieSceneActor>(Instance))
+	if (ALevelSequenceActor* LevelSequenceActor = Cast<ALevelSequenceActor>(Instance))
 	{
-		return MovieSceneActor->LevelSequence.TryLoad();
+		return LevelSequenceActor->LevelSequence.TryLoad();
 	}
 
 	return nullptr;
diff --git a/Engine/Source/Runtime/LevelSequence/Private/LevelSequenceActor.cpp b/Engine/Source/Runtime/LevelSequence/Private/LevelSequenceActor.cpp
new file mode 100644
--- /dev/null
+++ b/Engine/Source/Runtime/LevelSequence/Private/LevelSequenceActor.cpp
@@ -0,0 +1,129 @@
+// Copyright 1998-2015 Epic Games, Inc. All Rights Reserved.
+
+#include "LevelSequencePCH.h"
+#include "LevelSequenceActor.h"
+
+
+ALevelSequenceActor::ALevelSequenceActor(const FObjectInitializer& Init)
+	: Super(Init)
+{
+#if WITH_EDITORONLY_DATA
+	UBillboardComponent* SpriteComponent = CreateEditorOnlyDefaultSubobject<UBillboardComponent>(TEXT("Sprite"));
+
+	if (!IsRunningCommandlet())
+	{
+		// Structure to hold one-time initialization
+		struct FConstructorStatics
+		{
+			ConstructorHelpers::FObjectFinderOptional<UTexture2D> DecalTexture;
+			FConstructorStatics() : DecalTexture(TEXT("/Engine/EditorResources/SceneManager")) {}
+		};
+		static FConstructorStatics ConstructorStatics;
+
+		if (SpriteComponent)
+		{
+			SpriteComponent->Sprite = ConstructorStatics.DecalTexture.Get();
+			SpriteComponent->AttachParent = RootComponent;
+			SpriteComponent->bIsScreenSizeScaled = true;
+			SpriteComponent->bAbsoluteScale = true;
+			SpriteComponent->bReceivesDecals = false;
+			SpriteComponent->bHiddenInGame = true;
+			SetRootComponent(SpriteComponent);
+		}
+	}
+#endif // WITH_EDITORONLY_DATA
+
+	PrimaryActorTick.bCanEverTick = true;
+	bAutoPlay = false;
+}
+
+void ALevelSequenceActor::BeginPlay()
+{
+	Super::BeginPlay();
+	InitializePlayer();
+}
+
+void ALevelSequenceActor::PostInitProperties()
+{
+	Super::PostInitProperties();
+
+	if (!HasAnyFlags(RF_ClassDefaultObject))
+	{
+		// Make an instance of our asset so that we can keep hard references to actors we are using
+		SequenceInstance = NewObject<ULevelSequenceInstance>(this, "AnimationInstance");
+	}
+}
+
+void ALevelSequenceActor::PostLoad()
+{
+	Super::PostLoad();
+	
+	UpdateAnimationInstance();
+}
+
+#if WITH_EDITOR
+void ALevelSequenceActor::PostEditChangeProperty( struct FPropertyChangedEvent& PropertyChangedEvent)
+{
+	FName PropertyName = PropertyChangedEvent.Property ? PropertyChangedEvent.Property->GetFName() : NAME_None;
+
+	if (PropertyName == GET_MEMBER_NAME_CHECKED(ALevelSequenceActor, LevelSequence))
+	{
+		UpdateAnimationInstance();
+	}
+
+	Super::PostEditChangeProperty(PropertyChangedEvent);
+}
+
+bool ALevelSequenceActor::GetReferencedContentObjects(TArray<UObject*>& Objects) const
+{
+	if (UObject* Asset = SequenceInstance ? SequenceInstance->GetLevelSequence() : nullptr)
+	{
+		// @todo: Enable editing of animation instances
+		Objects.Add(Asset);
+	}
+
+	Super::GetReferencedContentObjects(Objects);
+
+	return true;
+}
+
+#endif // WITH_EDITOR
+
+void ALevelSequenceActor::Tick(float DeltaSeconds)
+{
+	if (SequencePlayer)
+	{
+		SequencePlayer->Update(DeltaSeconds);
+	}
+}
+
+void ALevelSequenceActor::SetSequence(ULevelSequence* InSequence)
+{
+	if (!SequencePlayer || !SequencePlayer->IsPlaying())
+	{
+		LevelSequence = InSequence;
+		InitializePlayer();
+	}
+}
+
+void ALevelSequenceActor::InitializePlayer()
+{
+	UpdateAnimationInstance();
+
+	if (GetWorld()->IsGameWorld())
+	{
+		SequencePlayer = NewObject<ULevelSequencePlayer>(this, "AnimationPlayer");
+		SequencePlayer->Initialize(SequenceInstance, GetWorld(), PlaybackSettings);
+
+		if (bAutoPlay)
+		{
+			SequencePlayer->Play();
+		}
+	}
+}
+
+void ALevelSequenceActor::UpdateAnimationInstance()
+{
+	ULevelSequence* LevelSequenceAsset = Cast<ULevelSequence>(LevelSequence.TryLoad());
+	SequenceInstance->Initialize(LevelSequenceAsset, GetWorld(), true);
+}
\ No newline at end of file
diff --git a/Engine/Source/Runtime/LevelSequence/Private/MovieSceneActor.cpp b/Engine/Source/Runtime/LevelSequence/Private/MovieSceneActor.cpp
deleted file mode 100644
diff --git a/Engine/Source/Runtime/LevelSequence/Public/LevelSequenceActor.h b/Engine/Source/Runtime/LevelSequence/Public/LevelSequenceActor.h
new file mode 100644
--- /dev/null
+++ b/Engine/Source/Runtime/LevelSequence/Public/LevelSequenceActor.h
@@ -0,0 +1,64 @@
+// Copyright 1998-2015 Epic Games, Inc. All Rights Reserved.
+
+#pragma once
+
+#include "LevelSequencePlayer.h"
+#include "LevelSequenceActor.generated.h"
+
+
+/**
+ * Actor responsible for controlling a specific level sequence in the world.
+ */
+UCLASS(Experimental, hideCategories=(Rendering, Physics, LOD, Activation))
+class LEVELSEQUENCE_API ALevelSequenceActor
+	: public AActor
+{
+public:
+
+	GENERATED_BODY()
+
+	/** Create and initialize a new instance. */
+	ALevelSequenceActor(const FObjectInitializer& Init);
+
+public:
+
+	UPROPERTY(EditAnywhere, BlueprintReadOnly, Category="Playback")
+	bool bAutoPlay;
+
+	UPROPERTY(EditAnywhere, BlueprintReadOnly, Category="Playback", meta=(ShowOnlyInnerProperties))
+	FLevelSequencePlaybackSettings PlaybackSettings;
+
+	UPROPERTY(transient, BlueprintReadOnly, Category="Playback")
+	ULevelSequencePlayer* SequencePlayer;
+
+	UPROPERTY(EditAnywhere, BlueprintReadOnly, Category="General", meta=(AllowedClasses="LevelSequence"))
+	FStringAssetReference LevelSequence;
+
+public:
+
+	UFUNCTION(BlueprintCallable, Category="Game|Cinematic")
+	void SetSequence(ULevelSequence* InSequence);
+
+public:
+
+	virtual void BeginPlay() override;
+	virtual void Tick(float DeltaSeconds) override;
+	virtual void PostInitProperties() override;
+	virtual void PostLoad() override;
+
+#if WITH_EDITOR
+	virtual void PostEditChangeProperty( struct FPropertyChangedEvent& PropertyChangedEvent) override;
+	virtual bool GetReferencedContentObjects(TArray<UObject*>& Objects) const override;
+#endif //WITH_EDITOR
+
+	void InitializePlayer();
+
+private:
+
+	/** Instance of the level sequence we're using. Stored as a property so we can track live actor references. */
+	UPROPERTY(Instanced)
+	ULevelSequenceInstance* SequenceInstance;
+
+	/** Update the instance of our level sequence. */
+	void UpdateAnimationInstance();
+};
diff --git a/Engine/Source/Runtime/LevelSequence/Public/MovieSceneActor.h b/Engine/Source/Runtime/LevelSequence/Public/MovieSceneActor.h
deleted file mode 100644
