Fixed a crash due to level sequence instance not being initialized

Fixes UE-21713.

diff --git a/Engine/Source/Runtime/LevelSequence/Private/LevelSequenceActor.cpp b/Engine/Source/Runtime/LevelSequence/Private/LevelSequenceActor.cpp
--- a/Engine/Source/Runtime/LevelSequence/Private/LevelSequenceActor.cpp
+++ b/Engine/Source/Runtime/LevelSequence/Private/LevelSequenceActor.cpp
@@ -49,8 +49,7 @@ void ALevelSequenceActor::PostInitProperties()
 
 	if (!HasAnyFlags(RF_ClassDefaultObject))
 	{
-		// Make an instance of our asset so that we can keep hard references to actors we are using
-		SequenceInstance = NewObject<ULevelSequenceInstance>(this, "AnimationInstance");
+		CreateSequenceInstance();
 	}
 }
 
@@ -122,8 +121,19 @@ void ALevelSequenceActor::InitializePlayer()
 	}
 }
 
+void ALevelSequenceActor::CreateSequenceInstance()
+{
+	if (!SequenceInstance)
+	{
+		// Make an instance of our asset so that we can keep hard references to actors we are using
+		SequenceInstance = NewObject<ULevelSequenceInstance>(this, "AnimationInstance");
+	}
+}
+
 void ALevelSequenceActor::UpdateAnimationInstance()
 {
+	CreateSequenceInstance();
+	
 	ULevelSequence* LevelSequenceAsset = Cast<ULevelSequence>(LevelSequence.TryLoad());
 	SequenceInstance->Initialize(LevelSequenceAsset, GetWorld(), true);
 }
\ No newline at end of file
diff --git a/Engine/Source/Runtime/LevelSequence/Public/LevelSequenceActor.h b/Engine/Source/Runtime/LevelSequence/Public/LevelSequenceActor.h
--- a/Engine/Source/Runtime/LevelSequence/Public/LevelSequenceActor.h
+++ b/Engine/Source/Runtime/LevelSequence/Public/LevelSequenceActor.h
@@ -59,6 +59,9 @@ private:
 	UPROPERTY(Instanced)
 	ULevelSequenceInstance* SequenceInstance;
 
+	/** Create the level sequence instance required to play back our sequence */
+	void CreateSequenceInstance();
+
 	/** Update the instance of our level sequence. */
 	void UpdateAnimationInstance();
 };
