Sequencer: Fix crash on deleting a spawned actor.

diff --git a/Engine/Source/Runtime/LevelSequence/Private/LevelSequence.cpp b/Engine/Source/Runtime/LevelSequence/Private/LevelSequence.cpp
--- a/Engine/Source/Runtime/LevelSequence/Private/LevelSequence.cpp
+++ b/Engine/Source/Runtime/LevelSequence/Private/LevelSequence.cpp
@@ -254,14 +254,6 @@ void ULevelSequenceInstance::SpawnOrDestroyObjects(bool DestroyAll)
 
 			if (ensure(World != nullptr))
 			{
-				if (!bAnyLevelActorsChanged)
-				{
-					DeselectAllActors();
-				}
-
-				// Actor should never be selected in the Editor at this point (we took care of that up above)
-				ensure(!Actor->IsSelected());
-
 				const bool bNetForce = false;
 				const bool bShouldModifyLevel = false;
 				const bool bWasDestroyed = World->DestroyActor(Actor, bNetForce, bShouldModifyLevel);
@@ -383,26 +375,3 @@ void ULevelSequenceInstance::UnbindPossessableObjects(const FGuid& ObjectId)
 
 	CachedObjectBindings.Remove(ObjectId);
 }
-
-void ULevelSequenceInstance::DeselectAllActors()
-{
-	// @todo sequencer: gmp: fix me
-	/*
-	USelection* Selection = GEditor->GetSelectedActors();
-	{
-		Selection->BeginBatchSelectOperation();
-		Selection->MarkBatchDirty();
-
-		for (auto SpawnedActorPair : SpawnedActors)
-		{
-			AActor* SpawnedActor = SpawnedActorPair.Value.Get();
-
-			if (SpawnedActor != nullptr)
-			{
-				GEditor->GetSelectedActors()->Deselect(SpawnedActor);
-			}
-		}
-
-		Selection->EndBatchSelectOperation();
-	}*/
-}
\ No newline at end of file
diff --git a/Engine/Source/Runtime/LevelSequence/Public/LevelSequence.h b/Engine/Source/Runtime/LevelSequence/Public/LevelSequence.h
--- a/Engine/Source/Runtime/LevelSequence/Public/LevelSequence.h
+++ b/Engine/Source/Runtime/LevelSequence/Public/LevelSequence.h
@@ -100,15 +100,6 @@ public:
 	virtual FText GetDisplayName() const override;
 #endif
 
-protected:
-
-	/**
-	 * Deselect all proxy actors in the Editor.
-	 *
-	 * @todo sequencer: remove Editor dependencies from this class
-	 */
-	void DeselectAllActors();
-
 private:
 
 	/** The level sequence we are instancing. */
