Copying //UE4/Dev-Sequencer to //UE4/Dev-Main (Source: //UE4/Dev-Sequencer @ 3178529)

diff --git a/Engine/Source/Runtime/Engine/Public/Curves/KeyFrameAlgorithms.h b/Engine/Source/Runtime/Engine/Public/Curves/KeyFrameAlgorithms.h
new file mode 100644
--- /dev/null
+++ b/Engine/Source/Runtime/Engine/Public/Curves/KeyFrameAlgorithms.h
@@ -0,0 +1,66 @@
+#pragma once
+
+#include "KeyFrameManipulator.h"
+
+namespace KeyFrameAlgorithms
+{
+	template<typename TimeType>
+	void Scale(TKeyFrameManipulator<TimeType>& Manipulator, TimeType ScaleOrigin, TimeType ScaleFactor)
+	{
+		TKeyTimeIterator<TimeType> Iterator = Manipulator.IterateKeys();
+
+		TArray<FKeyHandle> AllKeyHandles;
+		AllKeyHandles.Reserve(Iterator.GetEndIndex() - Iterator.GetStartIndex());
+
+		for (auto It = begin(Iterator); It != end(Iterator); ++It)
+		{
+			AllKeyHandles.Add(It.GetKeyHandle());
+		}
+
+		Scale(Manipulator, ScaleOrigin, ScaleFactor, AllKeyHandles);
+	}
+
+	template<typename TimeType, typename KeyHandleContainer>
+	void Scale(TKeyFrameManipulator<TimeType>& Manipulator, TimeType ScaleOrigin, TimeType ScaleFactor, const KeyHandleContainer& KeyHandles)
+	{
+		// @todo: we could be cleverer here by iterating forwards and backards from the scale origin time to prevent too much shuffling in the internal key array
+		for (FKeyHandle KeyHandle : KeyHandles)
+		{
+			TOptional<TimeType> Time = Manipulator.GetKeyTime(KeyHandle);
+			if (!Time.IsSet())
+			{
+				continue;
+			}
+			Manipulator.SetKeyTime(KeyHandle, (Time.GetValue() - ScaleOrigin) * ScaleFactor + ScaleOrigin);
+		}
+	}
+
+	template<typename TimeType>
+	void Translate(TKeyFrameManipulator<TimeType>& Manipulator, TimeType Delta)
+	{
+		TKeyTimeIterator<TimeType> Iterator = Manipulator.IterateKeys();
+
+		TSet<FKeyHandle> AllKeyHandles;
+		for (auto It = begin(Iterator); It != end(Iterator); ++It)
+		{
+			AllKeyHandles.Add(It.GetKeyHandle());
+		}
+
+		Translate(Manipulator, Delta, AllKeyHandles);
+	}
+
+	template<typename TimeType, typename KeyHandleContainer>
+	void Translate(TKeyFrameManipulator<TimeType>& Manipulator, TimeType Delta, const KeyHandleContainer& KeyHandles)
+	{
+		for (FKeyHandle KeyHandle : KeyHandles)
+		{
+			TOptional<TimeType> Time = Manipulator.GetKeyTime(KeyHandle);
+			if (!Time.IsSet())
+			{
+				continue;
+			}
+			Manipulator.SetKeyTime(KeyHandle, Time.GetValue() + Delta);
+		}
+	}
+
+}
\ No newline at end of file
diff --git a/Engine/Source/Runtime/EngineMessages/Private/EngineMessagesModule.cpp b/Engine/Source/Runtime/EngineMessages/Private/EngineMessagesModule.cpp
--- a/Engine/Source/Runtime/EngineMessages/Private/EngineMessagesModule.cpp
+++ b/Engine/Source/Runtime/EngineMessages/Private/EngineMessagesModule.cpp
@@ -11,12 +11,12 @@ class FEngineMessagesModule
 {
 public:
 
-	// IModuleInterface interface
+	//~ IModuleInterface interface
 
-	virtual void StartupModule( ) override { }
-	virtual void ShutdownModule( ) override { }
+	virtual void StartupModule() override { }
+	virtual void ShutdownModule() override { }
 
-	virtual bool SupportsDynamicReloading( ) override
+	virtual bool SupportsDynamicReloading() override
 	{
 		return true;
 	}
diff --git a/Engine/Source/Runtime/EngineMessages/Private/EngineMessagesPrivatePCH.h b/Engine/Source/Runtime/EngineMessages/Private/EngineMessagesPrivatePCH.h
--- a/Engine/Source/Runtime/EngineMessages/Private/EngineMessagesPrivatePCH.h
+++ b/Engine/Source/Runtime/EngineMessages/Private/EngineMessagesPrivatePCH.h
@@ -2,10 +2,4 @@
 
 #pragma once
 
-#include "EngineMessages.h"
-
-
-/* Private dependencies
- *****************************************************************************/
-
-#include "ModuleManager.h"
+#include "Runtime/Core/Public/Modules/ModuleManager.h"
diff --git a/Engine/Source/Runtime/Launch/Private/LaunchEngineLoop.cpp b/Engine/Source/Runtime/Launch/Private/LaunchEngineLoop.cpp
--- a/Engine/Source/Runtime/Launch/Private/LaunchEngineLoop.cpp
+++ b/Engine/Source/Runtime/Launch/Private/LaunchEngineLoop.cpp
@@ -29,7 +29,6 @@
 
 #if WITH_EDITOR
 	#include "EditorStyle.h"
-	#include "ProfilerClient.h"
 	#include "RemoteConfigIni.h"
 	#include "EditorCommandLineUtils.h"
 
@@ -75,7 +74,7 @@
 
 #if !UE_BUILD_SHIPPING
 	#include "STaskGraph.h"
-	#include "ProfilerService.h"
+	#include "IProfilerServiceModule.h"
 #endif
 
 #if WITH_AUTOMATION_WORKER
diff --git a/Engine/Source/Runtime/LevelSequence/Private/LevelSequence.cpp b/Engine/Source/Runtime/LevelSequence/Private/LevelSequence.cpp
--- a/Engine/Source/Runtime/LevelSequence/Private/LevelSequence.cpp
+++ b/Engine/Source/Runtime/LevelSequence/Private/LevelSequence.cpp
@@ -18,6 +18,7 @@ ULevelSequence::ULevelSequence(const FObjectInitializer& ObjectInitializer)
 	: Super(ObjectInitializer)
 	, MovieScene(nullptr)
 {
+	bParentContextsAreSignificant = true;
 }
 
 void ULevelSequence::Initialize()
@@ -127,19 +128,19 @@ void ULevelSequence::BindPossessableObject(const FGuid& ObjectId, const FLevelSe
 	ObjectReferences.CreateBinding(ObjectId, ObjectReference);
 }
 
-bool ULevelSequence::CanPossessObject(UObject& Object) const
+bool ULevelSequence::CanPossessObject(UObject& Object, UObject* InPlaybackContext) const
 {
 	return Object.IsA<AActor>() || Object.IsA<UActorComponent>();
 }
 
-UObject* ULevelSequence::FindPossessableObject(const FGuid& ObjectId, UObject* Context) const
+void ULevelSequence::LocateBoundObjects(const FGuid& ObjectId, UObject* Context, TArray<UObject*, TInlineAllocator<1>>& OutObjects) const
 {
-	return Context ? ObjectReferences.ResolveBinding(ObjectId, Context) : nullptr;
-}
-
-FGuid ULevelSequence::FindPossessableObjectId(UObject& Object) const
-{
-	return ObjectReferences.FindBindingId(&Object, Object.GetWorld());
+	// @todo: support multiple bindings
+	UObject* Object = Context ? ObjectReferences.ResolveBinding(ObjectId, Context) : nullptr;
+	if (Object)
+	{
+		OutObjects.Add(Object);
+	}
 }
 
 UMovieScene* ULevelSequence::GetMovieScene() const
@@ -167,4 +168,4 @@ bool ULevelSequence::AllowsSpawnableObjects() const
 void ULevelSequence::UnbindPossessableObjects(const FGuid& ObjectId)
 {
 	ObjectReferences.RemoveBinding(ObjectId);
-}
+}
\ No newline at end of file
diff --git a/Engine/Source/Runtime/LevelSequence/Private/LevelSequencePlayer.cpp b/Engine/Source/Runtime/LevelSequence/Private/LevelSequencePlayer.cpp
--- a/Engine/Source/Runtime/LevelSequence/Private/LevelSequencePlayer.cpp
+++ b/Engine/Source/Runtime/LevelSequence/Private/LevelSequencePlayer.cpp
@@ -74,11 +74,11 @@ ULevelSequencePlayer::ULevelSequencePlayer(const FObjectInitializer& ObjectIniti
 	, bIsPlaying(false)
 	, bReversePlayback(false)
 	, TimeCursorPosition(0.0f)
-	, LastCursorPosition(0.0f)
 	, StartTime(0.f)
 	, EndTime(0.f)
 	, CurrentNumLoops(0)
 	, bHasCleanedUpSequence(false)
+	, bIsEvaluating(false)
 {
 	SpawnRegister = MakeShareable(new FLevelSequenceSpawnRegister);
 }
@@ -119,8 +119,29 @@ void ULevelSequencePlayer::Pause()
 {
 	if (bIsPlaying)
 	{
+		if (bIsEvaluating)
+		{
+			LatentActions.Add(ELatentAction::Pause);
+			return;
+		}
+
 		bIsPlaying = false;
 
+		// Evaluate the sequence at its current time, with a status of 'stopped' to ensure that animated state pauses correctly
+		{
+			TGuardValue<bool> Guard(bIsEvaluating, true);
+
+			UMovieSceneSequence* MovieSceneSequence = RootTemplateInstance.GetSequence(MovieSceneSequenceID::Root);
+			TOptional<float> FixedFrameInterval = MovieSceneSequence->GetMovieScene() ? MovieSceneSequence->GetMovieScene()->GetOptionalFixedFrameInterval() : TOptional<float>();
+
+			FMovieSceneEvaluationRange Range = PlayPosition.JumpTo(GetSequencePosition(), FixedFrameInterval);
+
+			const FMovieSceneContext Context(Range, EMovieScenePlayerStatus::Stopped);
+			RootTemplateInstance.Evaluate(Context, *this);
+		}
+
+		ApplyLatentActions();
+
 		if (OnPause.IsBound())
 		{
 			OnPause.Broadcast();
@@ -132,14 +153,17 @@ void ULevelSequencePlayer::Stop()
 {
 	if (bIsPlaying)
 	{
+		if (bIsEvaluating)
+		{
+			LatentActions.Add(ELatentAction::Stop);
+			return;
+		}
+
 		bIsPlaying = false;
 		TimeCursorPosition = PlaybackSettings.PlayRate < 0.f ? GetLength() : 0.f;
 		CurrentNumLoops = 0;
 
-		if (RootMovieSceneInstance.IsValid())
-		{
-			RootMovieSceneInstance->RestoreState(*this);
-		}
+		RootTemplateInstance.Finish(*this);
 
 		SetTickPrerequisites(false);
 
@@ -190,28 +214,25 @@ void ULevelSequencePlayer::SetTickPrerequisites(bool bAddTickPrerequisites)
 		return;
 	}
 
-	TArray<UMovieSceneSequence*> AllSequences;
-	GetDescendantSequences(LevelSequence, AllSequences);
-
-	for (auto Sequence : AllSequences)
+	// @todo: this should happen procedurally, rather than up front. The current approach will not scale well.
+	for (auto& Pair : RootTemplateInstance.GetHierarchy().AllSubSequenceData())
 	{
-		UMovieScene* MovieScene = Sequence->GetMovieScene();
-		if (MovieScene == nullptr)
+		UMovieScene* MovieScene = Pair.Value.Sequence ? Pair.Value.Sequence->GetMovieScene() : nullptr;
+		if (!MovieScene)
 		{
 			continue;
 		}
 
 		TArray<AActor*> ControlledActors;
+		FMovieSceneSequenceID SequenceID = Pair.Key;
 
 		for (int32 PossessableCount = 0; PossessableCount < MovieScene->GetPossessableCount(); ++PossessableCount)
 		{
 			FMovieScenePossessable& Possessable = MovieScene->GetPossessable(PossessableCount);
-				
-			UObject* PossessableObject = Sequence->FindPossessableObject(Possessable.GetGuid(), this);
-			if (PossessableObject != nullptr)
-			{
-				AActor* PossessableActor = Cast<AActor>(PossessableObject);
 
+			for (TWeakObjectPtr<> PossessableObject : FindBoundObjects(Possessable.GetGuid(), SequenceID))
+			{
+				AActor* PossessableActor = Cast<AActor>(PossessableObject.Get());
 				if (PossessableActor != nullptr)
 				{
 					ControlledActors.Add(PossessableActor);
@@ -219,6 +240,8 @@ void ULevelSequencePlayer::SetTickPrerequisites(bool bAddTickPrerequisites)
 			}
 		}
 
+		// @todo: should this happen on spawn, not here?
+		// @todo: does setting tick prerequisites on spawnable *templates* actually propagate to spawned instances?
 		for (int32 SpawnableCount = 0; SpawnableCount < MovieScene->GetSpawnableCount(); ++ SpawnableCount)
 		{
 			FMovieSceneSpawnable& Spawnable = MovieScene->GetSpawnable(SpawnableCount);
@@ -292,7 +315,12 @@ void ULevelSequencePlayer::PlayInternal()
 
 		// Update now
 		bPendingFirstUpdate = false;
-		UpdateMovieSceneInstance(TimeCursorPosition, TimeCursorPosition);
+
+		UMovieSceneSequence* MovieSceneSequence = RootTemplateInstance.GetSequence(MovieSceneSequenceID::Root);
+		TOptional<float> FixedFrameInterval = MovieSceneSequence->GetMovieScene() ? MovieSceneSequence->GetMovieScene()->GetOptionalFixedFrameInterval() : TOptional<float>();
+
+		FMovieSceneEvaluationRange Range = PlayPosition.JumpTo(GetSequencePosition(), FixedFrameInterval);
+		UpdateMovieSceneInstance(Range);
 
 		if (OnPlay.IsBound())
 		{
@@ -316,10 +344,10 @@ void ULevelSequencePlayer::StartPlayingNextTick()
 
 	// @todo Sequencer playback: Should we recreate the instance every time?
 	// We must not recreate the instance since it holds stateful information (such as which objects it has spawned). Recreating the instance would break any 
-	if (!RootMovieSceneInstance.IsValid())
+	// @todo: Is this still the case now that eval state is stored (correctly) in the player?
+	if (!RootTemplateInstance.IsValid())
 	{
-		RootMovieSceneInstance = MakeShareable(new FMovieSceneSequenceInstance(*LevelSequence));
-		RootMovieSceneInstance->RefreshInstance(*this);
+		RootTemplateInstance.Initialize(*LevelSequence, *this);
 	}
 
 	SetTickPrerequisites(true);
@@ -337,7 +365,6 @@ float ULevelSequencePlayer::GetPlaybackPosition() const
 void ULevelSequencePlayer::SetPlaybackPosition(float NewPlaybackPosition)
 {
 	UpdateTimeCursorPosition(NewPlaybackPosition);
-	UpdateMovieSceneInstance(TimeCursorPosition, LastCursorPosition);
 }
 
 float ULevelSequencePlayer::GetLength() const
@@ -367,6 +394,9 @@ void ULevelSequencePlayer::UpdateTimeCursorPosition(float NewPosition)
 {
 	float Length = GetLength();
 
+	UMovieSceneSequence* MovieSceneSequence = RootTemplateInstance.GetSequence(MovieSceneSequenceID::Root);
+	TOptional<float> FixedFrameInterval = MovieSceneSequence->GetMovieScene() ? MovieSceneSequence->GetMovieScene()->GetOptionalFixedFrameInterval() : TOptional<float>();
+
 	if (bPendingFirstUpdate)
 	{
 		NewPosition = TimeCursorPosition;
@@ -381,19 +411,25 @@ void ULevelSequencePlayer::UpdateTimeCursorPosition(float NewPosition)
 			++CurrentNumLoops;
 			const float Overplay = FMath::Fmod(NewPosition, Length);
 			TimeCursorPosition = Overplay < 0 ? Length + Overplay : Overplay;
-			LastCursorPosition = TimeCursorPosition;
+			
+			PlayPosition.Reset(Overplay < 0 ? Length : 0.f);
+			FMovieSceneEvaluationRange Range = PlayPosition.PlayTo(GetSequencePosition(), FixedFrameInterval);
 
-			SpawnRegister->ForgetExternallyOwnedSpawnedObjects(*this);
+			SpawnRegister->ForgetExternallyOwnedSpawnedObjects(State, *this);
 
+			UpdateMovieSceneInstance(Range);
 			return;
 		}
 
 		// stop playback
 		Stop();
+		return;
 	}
 
-	LastCursorPosition = TimeCursorPosition;
 	TimeCursorPosition = NewPosition;
+
+	FMovieSceneEvaluationRange Range = PlayPosition.PlayTo(NewPosition + StartTime, FixedFrameInterval);
+	UpdateMovieSceneInstance(Range);
 }
 
 bool ULevelSequencePlayer::ShouldStopOrLoop(float NewPosition)
@@ -419,6 +455,8 @@ void ULevelSequencePlayer::Initialize(ULevelSequence* InLevelSequence, UWorld* I
 	CurrentPlayer = this;
 	LevelSequence = InLevelSequence;
 
+	RootTemplateInstance.Initialize(*LevelSequence, *this);
+
 	World = InWorld;
 	PlaybackSettings = Settings;
 
@@ -436,15 +474,6 @@ void ULevelSequencePlayer::Initialize(ULevelSequence* InLevelSequence, UWorld* I
 /* IMovieScenePlayer interface
  *****************************************************************************/
 
-void ULevelSequencePlayer::GetRuntimeObjects(TSharedRef<FMovieSceneSequenceInstance> MovieSceneInstance, const FGuid& ObjectId, TArray<TWeakObjectPtr<UObject>>& OutObjects) const
-{
-	UObject* FoundObject = MovieSceneInstance->FindObject(ObjectId, *this);
-	if (FoundObject)
-	{
-		OutObjects.Add(FoundObject);
-	}
-}
-
 
 void ULevelSequencePlayer::UpdateCameraCut(UObject* CameraObject, UObject* UnlockIfCameraObject, bool bJumpCut)
 {
@@ -466,6 +495,8 @@ void ULevelSequencePlayer::UpdateCameraCut(UObject* CameraObject, UObject* Unloc
 	}
 
 	UCameraComponent* CameraComponent = MovieSceneHelpers::CameraComponentFromRuntimeObject(CameraObject);
+	
+	CachedCameraComponent = CameraComponent;
 
 	if (CameraObject == ViewTarget)
 	{
@@ -534,19 +565,6 @@ void ULevelSequencePlayer::SetPlaybackStatus(EMovieScenePlayerStatus::Type InPla
 {
 }
 
-void ULevelSequencePlayer::AddOrUpdateMovieSceneInstance(UMovieSceneSection& MovieSceneSection, TSharedRef<FMovieSceneSequenceInstance> InstanceToAdd)
-{
-}
-
-void ULevelSequencePlayer::RemoveMovieSceneInstance(UMovieSceneSection& MovieSceneSection, TSharedRef<FMovieSceneSequenceInstance> InstanceToRemove)
-{
-}
-
-TSharedRef<FMovieSceneSequenceInstance> ULevelSequencePlayer::GetRootMovieSceneSequenceInstance() const
-{
-	return RootMovieSceneInstance.ToSharedRef();
-}
-
 UObject* ULevelSequencePlayer::GetPlaybackContext() const
 {
 	return World.Get();
@@ -584,45 +602,62 @@ void ULevelSequencePlayer::Update(const float DeltaSeconds)
 	{
 		float PlayRate = bReversePlayback ? -PlaybackSettings.PlayRate : PlaybackSettings.PlayRate;
 		UpdateTimeCursorPosition(TimeCursorPosition + DeltaSeconds * PlayRate);
-		UpdateMovieSceneInstance(TimeCursorPosition, LastCursorPosition);
 	}
 	else if (!bHasCleanedUpSequence && ShouldStopOrLoop(TimeCursorPosition))
 	{
-		UpdateMovieSceneInstance(TimeCursorPosition, LastCursorPosition);
+		// Update the sequence one last time
+		UMovieSceneSequence* MovieSceneSequence = RootTemplateInstance.GetSequence(MovieSceneSequenceID::Root);
+		TOptional<float> FixedFrameInterval = MovieSceneSequence->GetMovieScene() ? MovieSceneSequence->GetMovieScene()->GetOptionalFixedFrameInterval() : TOptional<float>();
+		
+		FMovieSceneEvaluationRange Range = PlayPosition.PlayTo(GetSequencePosition(), FixedFrameInterval);
+		UpdateMovieSceneInstance(Range);
 
 		bHasCleanedUpSequence = true;
 		
-		SpawnRegister->ForgetExternallyOwnedSpawnedObjects(*this);
+		SpawnRegister->ForgetExternallyOwnedSpawnedObjects(State, *this);
 		SpawnRegister->CleanUp(*this);
 	}
 }
 
-void ULevelSequencePlayer::UpdateMovieSceneInstance(float CurrentPosition, float PreviousPosition)
+void ULevelSequencePlayer::UpdateMovieSceneInstance(FMovieSceneEvaluationRange InRange)
 {
-	if(RootMovieSceneInstance.IsValid())
 	{
-		float Position = CurrentPosition + StartTime;
-		float LastPosition = PreviousPosition + StartTime;
-		UMovieSceneSequence* MovieSceneSequence = RootMovieSceneInstance->GetSequence();
-		if (MovieSceneSequence != nullptr && 
-			MovieSceneSequence->GetMovieScene()->GetForceFixedFrameIntervalPlayback() &&
-			MovieSceneSequence->GetMovieScene()->GetFixedFrameInterval() > 0 )
+		TGuardValue<bool> Guard(bIsEvaluating, true);
+
+		const FMovieSceneContext Context(InRange, GetPlaybackStatus());
+		RootTemplateInstance.Evaluate(Context, *this);
+
+	#if WITH_EDITOR
+		OnLevelSequencePlayerUpdate.Broadcast(*this, Context.GetTime(), Context.GetPreviousTime());
+	#endif
+	}
+
+	ApplyLatentActions();
+}
+
+void ULevelSequencePlayer::ApplyLatentActions()
+{
+	// Swap to a stack array to ensure no reentrancy if we evaluate during a pause, for instance
+	TArray<ELatentAction> TheseActions;
+	Swap(TheseActions, LatentActions);
+
+	for (ELatentAction LatentAction : TheseActions)
+	{
+		switch(LatentAction)
 		{
-			float FixedFrameInterval = MovieSceneSequence->GetMovieScene()->GetFixedFrameInterval();
-			Position = UMovieScene::CalculateFixedFrameTime( Position, FixedFrameInterval );
-			LastPosition = UMovieScene::CalculateFixedFrameTime( LastPosition, FixedFrameInterval );
+		case ELatentAction::Stop:	Stop(); break;
+		case ELatentAction::Pause:	Pause(); break;
 		}
-		EMovieSceneUpdateData UpdateData(Position, LastPosition);
-		RootMovieSceneInstance->Update(UpdateData, *this);
-
-#if WITH_EDITOR
-		OnLevelSequencePlayerUpdate.Broadcast(*this, CurrentPosition, PreviousPosition);
-#endif
 	}
 }
 
 void ULevelSequencePlayer::TakeFrameSnapshot(FLevelSequencePlayerSnapshot& OutSnapshot) const
 {
+	if (!ensure(LevelSequence))
+	{
+		return;
+	}
+	
 	const float CurrentTime = StartTime + TimeCursorPosition;
 
 	OutSnapshot.Settings = SnapshotSettings;
@@ -632,6 +667,7 @@ void ULevelSequencePlayer::TakeFrameSnapshot(FLevelSequencePlayerSnapshot& OutSn
 
 	OutSnapshot.CurrentShotName = OutSnapshot.MasterName;
 	OutSnapshot.CurrentShotLocalTime = CurrentTime;
+	OutSnapshot.CameraComponent = CachedCameraComponent.IsValid() ? CachedCameraComponent.Get() : nullptr;
 
 	UMovieSceneCinematicShotTrack* ShotTrack = LevelSequence->GetMovieScene()->FindMasterTrack<UMovieSceneCinematicShotTrack>();
 	if (ShotTrack)
@@ -640,6 +676,11 @@ void ULevelSequencePlayer::TakeFrameSnapshot(FLevelSequencePlayerSnapshot& OutSn
 		UMovieSceneCinematicShotSection* ActiveShot = nullptr;
 		for (UMovieSceneSection* Section : ShotTrack->GetAllSections())
 		{
+			if (!ensure(Section))
+			{
+				continue;
+			}
+
 			if (Section->IsActive() && Section->GetRange().Contains(CurrentTime) && (!ActiveShot || Section->GetRowIndex() < ActiveShot->GetRowIndex()))
 			{
 				ActiveShot = Cast<UMovieSceneCinematicShotSection>(Section);
@@ -652,8 +693,8 @@ void ULevelSequencePlayer::TakeFrameSnapshot(FLevelSequencePlayerSnapshot& OutSn
 			const float ShotLowerBound = ActiveShot->GetSequence() 
 				? ActiveShot->GetSequence()->GetMovieScene()->GetPlaybackRange().GetLowerBoundValue() 
 				: 0;
-			const float ShotOffset = ActiveShot->StartOffset + ShotLowerBound - ActiveShot->PrerollTime;
-			const float ShotPosition = ShotOffset + (CurrentTime - (ActiveShot->GetStartTime() - ActiveShot->PrerollTime)) / ActiveShot->TimeScale;
+			const float ShotOffset = ActiveShot->Parameters.StartOffset + ShotLowerBound - ActiveShot->Parameters.PrerollTime;
+			const float ShotPosition = ShotOffset + (CurrentTime - (ActiveShot->GetStartTime() - ActiveShot->Parameters.PrerollTime)) / ActiveShot->Parameters.TimeScale;
 
 			OutSnapshot.CurrentShotName = ActiveShot->GetShotDisplayName();
 			OutSnapshot.CurrentShotLocalTime = ShotPosition;
@@ -661,7 +702,7 @@ void ULevelSequencePlayer::TakeFrameSnapshot(FLevelSequencePlayerSnapshot& OutSn
 	}
 }
 
-IMovieSceneSpawnRegister& ULevelSequencePlayer::GetSpawnRegister()
+FMovieSceneSpawnRegister& ULevelSequencePlayer::GetSpawnRegister()
 {
 	return *SpawnRegister;
 }
\ No newline at end of file
diff --git a/Engine/Source/Runtime/LevelSequence/Private/LevelSequenceSpawnRegister.cpp b/Engine/Source/Runtime/LevelSequence/Private/LevelSequenceSpawnRegister.cpp
--- a/Engine/Source/Runtime/LevelSequence/Private/LevelSequenceSpawnRegister.cpp
+++ b/Engine/Source/Runtime/LevelSequence/Private/LevelSequenceSpawnRegister.cpp
@@ -12,25 +12,9 @@
 
 static const FName SequencerActorTag(TEXT("SequencerActor"));
 
-UObject* FLevelSequenceSpawnRegister::SpawnObject(const FGuid& BindingId, FMovieSceneSequenceInstance& SequenceInstance, IMovieScenePlayer& Player)
+UObject* FLevelSequenceSpawnRegister::SpawnObject(FMovieSceneSpawnable& Spawnable, FMovieSceneSequenceIDRef TemplateID, IMovieScenePlayer& Player)
 {
-	FMovieSceneSpawnRegisterKey Key(BindingId, SequenceInstance);
-
-	FSpawnedObject* Existing = Register.Find(Key);
-	UObject* ObjectInstance = Existing ? Existing->Object.Get() : nullptr;
-	if (ObjectInstance)
-	{
-		return ObjectInstance;
-	}
-
-	// Find the spawnable definition
-	FMovieSceneSpawnable* Spawnable = SequenceInstance.GetSequence()->GetMovieScene()->FindSpawnable(BindingId);
-	if (!Spawnable)
-	{
-		return nullptr;
-	}
-
-	AActor* ObjectTemplate = Cast<AActor>(Spawnable->GetObjectTemplate());
+	AActor* ObjectTemplate = Cast<AActor>(Spawnable.GetObjectTemplate());
 	if (!ObjectTemplate)
 	{
 		return nullptr;
@@ -112,32 +96,9 @@ UObject* FLevelSequenceSpawnRegister::SpawnObject(const FGuid& BindingId, FMovie
 
 	SpawnedActor->FinishSpawning(SpawnTransform);
 
-	Register.Add(Key, FSpawnedObject(BindingId, *SpawnedActor, Spawnable->GetSpawnOwnership()));
-
-	SequenceInstance.OnObjectSpawned(BindingId, *SpawnedActor, Player);
 	return SpawnedActor;
 }
 
-void FLevelSequenceSpawnRegister::DestroySpawnedObject(const FGuid& BindingId, FMovieSceneSequenceInstance& SequenceInstance, IMovieScenePlayer& Player)
-{
-	FMovieSceneSpawnRegisterKey Key(BindingId, SequenceInstance);
-	
-	FSpawnedObject* Existing = Register.Find(Key);
-	UObject* SpawnedObject = Existing ? Existing->Object.Get() : nullptr;
-	if (SpawnedObject)
-	{
-		PreDestroyObject(*SpawnedObject, BindingId, SequenceInstance);
-		DestroySpawnedObject(*SpawnedObject);
-	}
-
-	SequenceInstance.OnSpawnedObjectDestroyed(BindingId, Player);
-	Register.Remove(Key);
-}
-
-void FLevelSequenceSpawnRegister::PreDestroyObject(UObject& Object, const FGuid& BindingId, FMovieSceneSequenceInstance& SequenceInstance)
-{
-}
-
 void FLevelSequenceSpawnRegister::DestroySpawnedObject(UObject& Object)
 {
 	AActor* Actor = Cast<AActor>(&Object);
@@ -165,70 +126,4 @@ void FLevelSequenceSpawnRegister::DestroySpawnedObject(UObject& Object)
 		const bool bShouldModifyLevel = false;
 		World->DestroyActor(Actor, bNetForce, bShouldModifyLevel);
 	}
-}
-
-void FLevelSequenceSpawnRegister::ForgetExternallyOwnedSpawnedObjects(IMovieScenePlayer& Player)
-{
-	for (auto It = Register.CreateIterator(); It; ++It)
-	{
-		if (It.Value().Ownership == ESpawnOwnership::External)
-		{
-			TSharedPtr<FMovieSceneSequenceInstance> SequenceInstance = It.Key().SequenceInstance.Pin();
-			if (SequenceInstance.IsValid())
-			{
-				SequenceInstance->OnSpawnedObjectDestroyed(It.Key().BindingId, Player);
-			}
-			It.RemoveCurrent();
-		}
-	}
-}
-
-void FLevelSequenceSpawnRegister::DestroyObjectsByPredicate(IMovieScenePlayer& Player, const TFunctionRef<bool(const FGuid&, ESpawnOwnership, FMovieSceneSequenceInstance&)>& Predicate)
-{
-	for (auto It = Register.CreateIterator(); It; ++It)
-	{
-		FMovieSceneSequenceInstance* ThisInstance = It.Key().SequenceInstance.Pin().Get();
-		if (!ThisInstance || Predicate(It.Value().Guid, It.Value().Ownership, *ThisInstance))
-		{
-			UObject* SpawnedObject = It.Value().Object.Get();
-			if (SpawnedObject)
-			{
-				if (ThisInstance)
-				{
-					PreDestroyObject(*SpawnedObject, It.Key().BindingId, *ThisInstance);
-				}
-				DestroySpawnedObject(*SpawnedObject);
-			}
-
-			if (ThisInstance)
-			{
-				ThisInstance->OnSpawnedObjectDestroyed(It.Key().BindingId, Player);
-			}
-			It.RemoveCurrent();
-		}
-	}
-}
-
-void FLevelSequenceSpawnRegister::PreUpdateSequenceInstance(FMovieSceneSequenceInstance& Instance, IMovieScenePlayer& Player)
-{
-	++CurrentlyUpdatingSequenceCount;
-	ActiveInstances.Add(Instance.AsShared());
-}
-
-void FLevelSequenceSpawnRegister::PostUpdateSequenceInstance(FMovieSceneSequenceInstance& InInstance, IMovieScenePlayer& Player)
-{
-	if (--CurrentlyUpdatingSequenceCount == 0)
-	{
-		for (TWeakPtr<FMovieSceneSequenceInstance>& WeakInstance : PreviouslyActiveInstances)
-		{
-			TSharedPtr<FMovieSceneSequenceInstance> Instance = WeakInstance.Pin();
-			if (Instance.IsValid() && !ActiveInstances.Contains(Instance))
-			{
-				OnSequenceExpired(*Instance, Player);
-			}
-		}
-
-		Swap(ActiveInstances, PreviouslyActiveInstances);
-		ActiveInstances.Reset();
-	}
 }
\ No newline at end of file
diff --git a/Engine/Source/Runtime/LevelSequence/Public/LevelSequence.h b/Engine/Source/Runtime/LevelSequence/Public/LevelSequence.h
--- a/Engine/Source/Runtime/LevelSequence/Public/LevelSequence.h
+++ b/Engine/Source/Runtime/LevelSequence/Public/LevelSequence.h
@@ -33,11 +33,9 @@ public:
 public:
 
 	// UMovieSceneSequence interface
-	
 	virtual void BindPossessableObject(const FGuid& ObjectId, UObject& PossessedObject, UObject* Context) override;
-	virtual bool CanPossessObject(UObject& Object) const override;
-	virtual UObject* FindPossessableObject(const FGuid& ObjectId, UObject* Context) const override;
-	virtual FGuid FindPossessableObjectId(UObject& Object) const override;
+	virtual bool CanPossessObject(UObject& Object, UObject* InPlaybackContext) const override;
+	virtual void LocateBoundObjects(const FGuid& ObjectId, UObject* Context, TArray<UObject*, TInlineAllocator<1>>& OutObjects) const override;
 	virtual UMovieScene* GetMovieScene() const override;
 	virtual UObject* GetParentObject(UObject* Object) const override;
 	virtual void UnbindPossessableObjects(const FGuid& ObjectId) override;
diff --git a/Engine/Source/Runtime/LevelSequence/Public/LevelSequencePlayer.h b/Engine/Source/Runtime/LevelSequence/Public/LevelSequencePlayer.h
--- a/Engine/Source/Runtime/LevelSequence/Public/LevelSequencePlayer.h
+++ b/Engine/Source/Runtime/LevelSequence/Public/LevelSequencePlayer.h
@@ -4,10 +4,10 @@
 
 #include "IMovieScenePlayer.h"
 #include "LevelSequence.h"
+#include "MovieSceneEvaluationTemplateInstance.h"
 #include "LevelSequencePlayer.generated.h"
 
 class FLevelSequenceSpawnRegister;
-class FMovieSceneSequenceInstance;
 class ULevel;
 class UMovieSceneBindings;
 
@@ -55,6 +55,9 @@ struct FLevelSequencePlayerSnapshot
 	UPROPERTY(VisibleAnywhere, BlueprintReadOnly, Category="General")
 	float CurrentShotLocalTime;
 
+	UPROPERTY(VisibleAnywhere, BlueprintReadOnly, Category="General")
+	UCameraComponent* CameraComponent;
+
 	UPROPERTY(VisibleAnywhere, BlueprintReadOnly, Category="General")
 	FLevelSequenceSnapshotSettings Settings;
 };
@@ -211,16 +214,13 @@ public:
 protected:
 
 	// IMovieScenePlayer interface
-	virtual void GetRuntimeObjects(TSharedRef<FMovieSceneSequenceInstance> MovieSceneInstance, const FGuid& ObjectHandle, TArray<TWeakObjectPtr<UObject>>& OutObjects) const override;
+	virtual FMovieSceneRootEvaluationTemplateInstance& GetEvaluationTemplate() override { return RootTemplateInstance; }
 	virtual void UpdateCameraCut(UObject* CameraObject, UObject* UnlockIfCameraObject, bool bJumpCut) override;
 	virtual void SetViewportSettings(const TMap<FViewportClient*, EMovieSceneViewportParams>& ViewportParamsMap) override;
 	virtual void GetViewportSettings(TMap<FViewportClient*, EMovieSceneViewportParams>& ViewportParamsMap) const override;
 	virtual EMovieScenePlayerStatus::Type GetPlaybackStatus() const override;
 	virtual void SetPlaybackStatus(EMovieScenePlayerStatus::Type InPlaybackStatus) override;
-	virtual void AddOrUpdateMovieSceneInstance(UMovieSceneSection& MovieSceneSection, TSharedRef<FMovieSceneSequenceInstance> InstanceToAdd) override;
-	virtual void RemoveMovieSceneInstance(UMovieSceneSection& MovieSceneSection, TSharedRef<FMovieSceneSequenceInstance> InstanceToRemove) override;
-	virtual TSharedRef<FMovieSceneSequenceInstance> GetRootMovieSceneSequenceInstance() const override;
-	virtual IMovieSceneSpawnRegister& GetSpawnRegister() override;
+	virtual FMovieSceneSpawnRegister& GetSpawnRegister() override;
 	virtual UObject* GetPlaybackContext() const override;
 	virtual TArray<UObject*> GetEventContexts() const override;
 
@@ -240,7 +240,7 @@ private:
 	void PlayInternal();
 
 	/** Update the movie scene instance from the specified previous position, to the specified time position. */
-	void UpdateMovieSceneInstance(float CurrentPosition, float PreviousPosition);
+	void UpdateMovieSceneInstance(FMovieSceneEvaluationRange InRange);
 
 	/** Update the time cursor position and handle stopping & looping. */
 	void UpdateTimeCursorPosition(float NewPosition);
@@ -251,7 +251,17 @@ private:
 	/** Returns whether playback should be stopped or looped because it has gone out of the playback bounds. */
 	bool ShouldStopOrLoop(float NewPosition);
 
+	/** Get the position of the cursor in the root sequence */
+	FORCEINLINE float GetSequencePosition() const { return TimeCursorPosition + StartTime; }
+
+	/** Apply any latent actions which may have accumulated while the sequence was being evaluated */
+	void ApplyLatentActions();
+
 private:
+	enum class ELatentAction
+	{
+		Stop, Pause
+	};
 
 	/** The level sequence to play. */
 	UPROPERTY(transient)
@@ -276,9 +286,6 @@ private:
 	UPROPERTY()
 	float TimeCursorPosition;
 
-	/** The time cursor position in the previous update. */
-	float LastCursorPosition;
-
 	/** Time time at which to start playing the sequence (defaults to the lower bound of the sequence's play range) */
 	float StartTime;
 
@@ -295,10 +302,13 @@ private:
 	/** Whether this player has cleaned up the level sequence after it has stopped playing or not */
 	bool bHasCleanedUpSequence;
 
-private:
+	/** Set to true while evaluating to prevent reentrancy */
+	bool bIsEvaluating;
 
-	/** The root movie scene instance to update when playing. */
-	TSharedPtr<FMovieSceneSequenceInstance> RootMovieSceneInstance;
+	/** Set of latent actions that are to be performed when the sequence has finished evaluating this frame */
+	TArray<ELatentAction> LatentActions;
+
+private:
 
 	/** The world this player will spawn actors in, if needed */
 	TWeakObjectPtr<UWorld> World;
@@ -309,6 +319,11 @@ private:
 	/** The last view target to reset to when updating camera cuts to null */
 	TWeakObjectPtr<AActor> LastViewTarget;
 
+	UPROPERTY(transient)
+	FMovieSceneRootEvaluationTemplateInstance RootTemplateInstance;
+
+	FMovieScenePlaybackPosition PlayPosition;
+
 protected:
 
 	/** How to take snapshots */
@@ -336,4 +351,6 @@ private:
 
 	/** The event that will be broadcast every time the sequence is updated */
 	mutable FOnLevelSequencePlayerUpdated OnLevelSequencePlayerUpdate;
+
+	TWeakObjectPtr<UCameraComponent> CachedCameraComponent;
 };
diff --git a/Engine/Source/Runtime/LevelSequence/Public/LevelSequenceSpawnRegister.h b/Engine/Source/Runtime/LevelSequence/Public/LevelSequenceSpawnRegister.h
--- a/Engine/Source/Runtime/LevelSequence/Public/LevelSequenceSpawnRegister.h
+++ b/Engine/Source/Runtime/LevelSequence/Public/LevelSequenceSpawnRegister.h
@@ -2,65 +2,13 @@
 
 #pragma once
 
-#include "IMovieSceneSpawnRegister.h"
-#include "MovieSceneSequenceInstance.h"
-#include "MovieSceneSpawnable.h"
+#include "MovieSceneSpawnRegister.h"
 
 /** Movie scene spawn register that knows how to handle spawning objects (actors) for a level sequence  */
-class LEVELSEQUENCE_API FLevelSequenceSpawnRegister : public IMovieSceneSpawnRegister
+class LEVELSEQUENCE_API FLevelSequenceSpawnRegister : public FMovieSceneSpawnRegister
 {
-public:
-	FLevelSequenceSpawnRegister()
-	{
-		CurrentlyUpdatingSequenceCount = 0;
-	}
-
-	/** ~ IMovieSceneSpawnRegister interface */
-	virtual UObject* SpawnObject(const FGuid& BindingId, FMovieSceneSequenceInstance& SequenceInstance, IMovieScenePlayer&) override;
-	virtual void DestroySpawnedObject(const FGuid& BindingId, FMovieSceneSequenceInstance& SequenceInstance, IMovieScenePlayer&) override;
-	virtual void DestroyObjectsByPredicate(IMovieScenePlayer& Player, const TFunctionRef<bool(const FGuid&, ESpawnOwnership, FMovieSceneSequenceInstance&)>& Predicate) override;
-	virtual void ForgetExternallyOwnedSpawnedObjects(IMovieScenePlayer& Player) override;
-	virtual void PreUpdateSequenceInstance(FMovieSceneSequenceInstance& Instance, IMovieScenePlayer&) override;
-	virtual void PostUpdateSequenceInstance(FMovieSceneSequenceInstance& Instance, IMovieScenePlayer&) override;
-
-protected:
-
-	/** Structure holding information pertaining to a spawned object */
-	struct FSpawnedObject
-	{
-		FSpawnedObject(const FGuid& InGuid, UObject& InObject, ESpawnOwnership InOwnership)
-			: Guid(InGuid)
-			, Object(&InObject)
-			, Ownership(InOwnership)
-		{}
-
-		/** The ID of the sequencer object binding that this object relates to */
-		FGuid Guid;
-
-		/** The object that has been spawned */
-		TWeakObjectPtr<UObject> Object;
-
-		/** What level of ownership this object was spawned with */
-		ESpawnOwnership Ownership;
-	};
-
 protected:
-	/** Called right before a spawned object with the specified ID and sequence instance is destroyed */
-	virtual void PreDestroyObject(UObject& Object, const FGuid& BindingId, FMovieSceneSequenceInstance& SequenceInstance);
-
-	/** Implementation function that destroys the specified object if it is an actor */
-	void DestroySpawnedObject(UObject& Object);
-
-	/** Destroy any objects that pass the specified predicate */
-	void DestroyObjects(IMovieScenePlayer& Player, TFunctionRef<bool(FMovieSceneSequenceInstance&, FSpawnedObject&)> Pred);
-	
-	
-	/** Sets of sequence instances that are active for the current update, and were active for the previous update */
-	TSet<TWeakPtr<FMovieSceneSequenceInstance>> ActiveInstances, PreviouslyActiveInstances;
-
-	/** The number of currently updating sequence instances. Used to keep track of the outermost sequence instance that triggered an update */
-	int32 CurrentlyUpdatingSequenceCount;
-
-	/** Register of spawned objects */
-	TMap<FMovieSceneSpawnRegisterKey, FSpawnedObject> Register;
+	/** ~ FMovieSceneSpawnRegister interface */
+	virtual UObject* SpawnObject(FMovieSceneSpawnable& Spawnable, FMovieSceneSequenceIDRef TemplateID, IMovieScenePlayer& Player) override;
+	virtual void DestroySpawnedObject(UObject& Object) override;
 };
\ No newline at end of file
diff --git a/Engine/Source/Runtime/Media/Public/IMediaTextureSink.h b/Engine/Source/Runtime/Media/Public/IMediaTextureSink.h
--- a/Engine/Source/Runtime/Media/Public/IMediaTextureSink.h
+++ b/Engine/Source/Runtime/Media/Public/IMediaTextureSink.h
@@ -3,8 +3,9 @@
 #pragma once
 
 
-class FRHITexture;
-class FRHITexture2D;
+#if WITH_ENGINE
+	class FRHITexture;
+#endif
 
 
 /**
@@ -18,6 +19,9 @@ enum class EMediaTextureSinkFormat
 	/** Four 8-bit unsigned integer components (Blue, Green, Red, Alpha) per pixel. */
 	CharBGRA,
 
+	/**  Windows bitmap (like CharBGRA, but flipped vertically). */
+	CharBMP,
+
 	/** NV12 encoded monochrome texture with 8 bits per channel. */
 	CharNV12,
 
@@ -135,13 +139,14 @@ public:
 	/**
 	 * Initialize the sink.
 	 *
-	 * @param Dimensions Width and height of the texture (in pixels).
+	 * @param OutputDim Width and height of the video output (in pixels).
+	 * @param BufferDim Width and height of the sink buffer(s) (in pixels).
 	 * @param Format The pixel format of the sink's render target texture.
 	 * @param Mode The mode to operate the sink in (buffered vs. unbuffered).
 	 * @return true on success, false if initialization is already in progress or failed.
 	 * @see GetTextureSinkDimensions, GetTextureSinkMode, ShutdownSink, SupportsTextureSinkFormat
 	 */
-	virtual bool InitializeTextureSink(FIntPoint Dimensions, EMediaTextureSinkFormat Format, EMediaTextureSinkMode Mode) = 0;
+	virtual bool InitializeTextureSink(FIntPoint OutputDim, FIntPoint BufferDim, EMediaTextureSinkFormat Format, EMediaTextureSinkMode Mode) = 0;
 
 	/**
 	 * Shut down the sink.
@@ -177,7 +182,7 @@ public:
 	 */
 	virtual FRHITexture* GetTextureSinkTexture() = 0;
 
-#endif
+#endif //WITH_ENGINE
 
 public:
 
@@ -277,7 +282,38 @@ public:
 	 */
 	virtual void UpdateTextureSinkResource(FRHITexture* RenderTarget, FRHITexture* ShaderResource) = 0;
 
-#endif
+#endif //WITH_ENGINE
+
+public:
+
+	DEPRECATED(4.15, "This overload of InitializeTextureSink has been deprecated. Please use the other one!")
+	bool InitializeTextureSink(FIntPoint Dimensions, EMediaTextureSinkFormat Format, EMediaTextureSinkMode Mode)
+	{
+		FIntPoint BufferDim;
+
+		switch (Format)
+		{
+		case EMediaTextureSinkFormat::CharNV12:
+		case EMediaTextureSinkFormat::CharNV21:
+			BufferDim = FIntPoint(Align(Dimensions.X, 16), Align(Dimensions.Y, 2) * 3 / 2);
+			break;
+
+		case EMediaTextureSinkFormat::CharUYVY:
+		case EMediaTextureSinkFormat::CharYUY2:
+		case EMediaTextureSinkFormat::CharYVYU:
+			BufferDim = FIntPoint(Align(Dimensions.X, 16) / 2, Dimensions.Y);
+			break;
+
+		case EMediaTextureSinkFormat::CharAYUV:
+		case EMediaTextureSinkFormat::CharBGRA:
+		case EMediaTextureSinkFormat::FloatRGB:
+		case EMediaTextureSinkFormat::FloatRGBA:
+			BufferDim = Dimensions;
+			break;
+		}
+
+		return InitializeTextureSink(Dimensions, BufferDim, Format, Mode);
+	}
 
 public:
 
diff --git a/Engine/Source/Runtime/MediaAssets/Private/Assets/MediaPlayer.cpp b/Engine/Source/Runtime/MediaAssets/Private/Assets/MediaPlayer.cpp
--- a/Engine/Source/Runtime/MediaAssets/Private/Assets/MediaPlayer.cpp
+++ b/Engine/Source/Runtime/MediaAssets/Private/Assets/MediaPlayer.cpp
@@ -28,7 +28,14 @@ UMediaPlayer::UMediaPlayer(const FObjectInitializer& ObjectInitializer)
 
 bool UMediaPlayer::CanPause() const
 {
-	return Player.IsValid() && (Player->GetControls().GetState() == EMediaState::Playing);
+	if (!Player.IsValid())
+	{
+		return false;
+	}
+	
+	EMediaState State = Player->GetControls().GetState();
+
+	return ((State== EMediaState::Playing) || (State== EMediaState::Preparing));
 }
 
 
@@ -201,7 +208,7 @@ bool UMediaPlayer::Next()
 
 	int32 RemainingAttempts = Playlist->Num();
 
-	while (--RemainingAttempts >= 0)
+	while (RemainingAttempts-- > 0)
 	{
 		UMediaSource* NextSource = Shuffle
 			? Playlist->GetRandom(PlaylistIndex)
@@ -355,9 +362,9 @@ bool UMediaPlayer::Rewind()
 }
 
 
-bool UMediaPlayer::Seek(const FTimespan& InTime)
+bool UMediaPlayer::Seek(const FTimespan& Time)
 {
-	return Player.IsValid() && Player->GetControls().Seek(InTime);
+	return Player.IsValid() && Player->GetControls().Seek(Time);
 }
 
 
@@ -367,11 +374,11 @@ bool UMediaPlayer::SelectTrack(EMediaPlayerTrack TrackType, int32 TrackIndex)
 }
 
 
-bool UMediaPlayer::SetLooping(bool InLooping)
+bool UMediaPlayer::SetLooping(bool Looping)
 {
-	Loop = InLooping;
+	Loop = Looping;
 
-	return Player.IsValid() && Player->GetControls().SetLooping(InLooping);
+	return Player.IsValid() && Player->GetControls().SetLooping(Looping);
 }
 
 
diff --git a/Engine/Source/Runtime/MediaAssets/Private/Assets/MediaSoundWave.cpp b/Engine/Source/Runtime/MediaAssets/Private/Assets/MediaSoundWave.cpp
--- a/Engine/Source/Runtime/MediaAssets/Private/Assets/MediaSoundWave.cpp
+++ b/Engine/Source/Runtime/MediaAssets/Private/Assets/MediaSoundWave.cpp
@@ -55,10 +55,10 @@ int32 UMediaSoundWave::GeneratePCMData(uint8* Data, const int32 SamplesRequested
 
 		// poor man's re-sampling
 		const float ResampleRatio = (float)SinkSampleRate / 44100.0f;
-		const int32 SamplesAvailable = (QueuedAudio.Num() / sizeof(int16)) / ResampleRatio - 2 * SinkNumChannels;
-		const int32 NumFrames = FMath::Min<int32>(SamplesRequested, SamplesAvailable) / NumChannels;
+		const int32 SamplesAvailable = FPlatformMath::FloorToInt(((QueuedAudio.Num() / sizeof(int16)) - SinkNumChannels) / ResampleRatio);
+		const int32 NumOutputFrames = FMath::Min<int32>(SamplesRequested, SamplesAvailable) / SinkNumChannels;
 		
-		if (NumFrames > 0)
+		if (NumOutputFrames > 0)
 		{
 			int16* InputSamples = (int16*)QueuedAudio.GetData();
 			int16* OutputSamples = (int16*)Data;
@@ -66,7 +66,7 @@ int32 UMediaSoundWave::GeneratePCMData(uint8* Data, const int32 SamplesRequested
 			float InputFrameInterp = 0.0f;
 			int32 InputIndex = 0;
 
-			for (int32 OutputFrame = 0; OutputFrame < NumFrames; ++OutputFrame)
+			for (int32 OutputFrame = 0; OutputFrame < NumOutputFrames; ++OutputFrame)
 			{
 				int32 InputFrame = (int32)InputFrameInterp;
 				float Alpha = InputFrameInterp - InputFrame;
@@ -90,6 +90,17 @@ int32 UMediaSoundWave::GeneratePCMData(uint8* Data, const int32 SamplesRequested
 					OutputSamples[0] = OutputSampleL;
 					OutputSamples[1] = OutputSampleR;
 				}
+				else if (SinkNumChannels == 4)
+				{
+					// quadrophonic to stereo
+					const int16 OutputSampleLF = (int16)FMath::Lerp((float)InputSamples[InputIndex], (float)InputSamples[InputIndex + 4], Alpha);
+					const int16 OutputSampleRF = (int16)FMath::Lerp((float)InputSamples[InputIndex + 1], (float)InputSamples[InputIndex + 5], Alpha);
+					const int16 OutputSampleLB = (int16)FMath::Lerp((float)InputSamples[InputIndex + 2], (float)InputSamples[InputIndex + 6], Alpha);
+					const int16 OutputSampleRB = (int16)FMath::Lerp((float)InputSamples[InputIndex + 3], (float)InputSamples[InputIndex + 7], Alpha);
+
+					OutputSamples[0] = (OutputSampleLF + OutputSampleRF) / 2;
+					OutputSamples[1] = (OutputSampleLB + OutputSampleRB) / 2;
+				}
 				else if (SinkNumChannels == 6)
 				{
 					// 5.1 to stereo
@@ -209,9 +220,9 @@ bool UMediaSoundWave::InitializeAudioSink(uint32 Channels, uint32 InSampleRate)
 {
 	UE_LOG(LogMediaAssets, Verbose, TEXT("MediaSoundWave initializing sink with %i channels at %i Hz."), Channels, InSampleRate);
 
-	if ((Channels != 1) && (Channels != 2) && (Channels != 6))
+	if ((Channels != 1) && (Channels != 2) && (Channels != 4) && (Channels != 6))
 	{
-		return false; // we currently only support mono, stereo, and 5.1
+		return false; // we currently only support mono, stereo, quadrophonic, and 5.1
 	}
 
 	if ((Channels == SinkNumChannels) && (InSampleRate == SinkSampleRate))
diff --git a/Engine/Source/Runtime/MediaAssets/Private/Assets/MediaTexture.cpp b/Engine/Source/Runtime/MediaAssets/Private/Assets/MediaTexture.cpp
--- a/Engine/Source/Runtime/MediaAssets/Private/Assets/MediaTexture.cpp
+++ b/Engine/Source/Runtime/MediaAssets/Private/Assets/MediaTexture.cpp
@@ -13,7 +13,8 @@ UMediaTexture::UMediaTexture(const FObjectInitializer& ObjectInitializer)
 	, AddressX(TA_Clamp)
 	, AddressY(TA_Clamp)
 	, ClearColor(FLinearColor::Black)
-	, SinkDimensions(FIntPoint(1, 1))
+	, SinkBufferDim(FIntPoint(1, 1))
+	, SinkOutputDim(FIntPoint(1, 1))
 	, SinkFormat(EMediaTextureSinkFormat::CharBGRA)
 	, SinkMode(EMediaTextureSinkMode::Unbuffered)
 {
@@ -21,6 +22,39 @@ UMediaTexture::UMediaTexture(const FObjectInitializer& ObjectInitializer)
 }
 
 
+/* IMediaTextureSink interface
+ *****************************************************************************/
+
+float UMediaTexture::GetAspectRatio() const
+{
+	if (Resource == nullptr)
+	{
+		return 0.0f;
+	}
+
+	const FIntPoint Dimensions = ((FMediaTextureResource*)Resource)->GetSizeXY();
+
+	if (Dimensions.Y == 0)
+	{
+		return 0.0f;
+	}
+
+	return Dimensions.X / Dimensions.Y;
+}
+
+
+int32 UMediaTexture::GetHeight() const
+{
+	return (Resource != nullptr) ? (int32)Resource->GetSizeY() : 0;
+}
+
+
+int32 UMediaTexture::GetWidth() const
+{
+	return (Resource != nullptr) ? (int32)Resource->GetSizeX() : 0;
+}
+
+
 /* IMediaTextureSink interface
  *****************************************************************************/
 
@@ -71,11 +105,12 @@ FRHITexture* UMediaTexture::GetTextureSinkTexture()
 }
 
 
-bool UMediaTexture::InitializeTextureSink(FIntPoint Dimensions, EMediaTextureSinkFormat Format, EMediaTextureSinkMode Mode)
+bool UMediaTexture::InitializeTextureSink(FIntPoint OutputDim, FIntPoint BufferDim, EMediaTextureSinkFormat Format, EMediaTextureSinkMode Mode)
 {
-	UE_LOG(LogMediaAssets, Verbose, TEXT("MediaTexture initializing sink with %i x %i pixels %s."), Dimensions.X, Dimensions.Y, (Mode == EMediaTextureSinkMode::Buffered) ? TEXT("Buffered") : TEXT("Unbuffered"));
+	UE_LOG(LogMediaAssets, Verbose, TEXT("MediaTexture initializing sink with %ix%i output and %ix%i buffer as %s."), OutputDim.X, OutputDim.Y, BufferDim.X, BufferDim.Y, (Mode == EMediaTextureSinkMode::Buffered) ? TEXT("Buffered") : TEXT("Unbuffered"));
 
-	SinkDimensions = Dimensions;
+	SinkBufferDim = BufferDim;
+	SinkOutputDim = OutputDim;
 	SinkFormat = Format;
 	SinkMode = Mode;
 
@@ -86,7 +121,7 @@ bool UMediaTexture::InitializeTextureSink(FIntPoint Dimensions, EMediaTextureSin
 		return false;
 	}
 
-	((FMediaTextureResource*)Resource)->InitializeBuffer(Dimensions, Format, Mode);
+	((FMediaTextureResource*)Resource)->InitializeBuffer(OutputDim, BufferDim, Format, Mode);
 
 	return true;
 }
@@ -113,7 +148,8 @@ void UMediaTexture::ShutdownTextureSink()
 	}
 
 	// reset to 1x1 clear color
-	SinkDimensions = FIntPoint(1, 1);
+	SinkBufferDim = FIntPoint(1, 1);
+	SinkOutputDim = FIntPoint(1, 1);
 	SinkFormat = EMediaTextureSinkFormat::CharBGRA;
 	SinkMode = EMediaTextureSinkMode::Unbuffered;
 
@@ -121,7 +157,7 @@ void UMediaTexture::ShutdownTextureSink()
 
 	if (Resource != nullptr)
 	{
-		((FMediaTextureResource*)Resource)->InitializeBuffer(SinkDimensions, SinkFormat, SinkMode);
+		((FMediaTextureResource*)Resource)->InitializeBuffer(SinkOutputDim, SinkBufferDim, SinkFormat, SinkMode);
 	}
 }
 
@@ -159,7 +195,7 @@ void UMediaTexture::UpdateTextureSinkResource(FRHITexture* RenderTarget, FRHITex
 
 FTextureResource* UMediaTexture::CreateResource()
 {
-	return new FMediaTextureResource(*this, ClearColor, SinkDimensions, SinkFormat, SinkMode);
+	return new FMediaTextureResource(*this, ClearColor, SinkOutputDim, SinkFormat, SinkMode);
 }
 
 
diff --git a/Engine/Source/Runtime/MediaAssets/Private/Misc/MediaTextureResource.cpp b/Engine/Source/Runtime/MediaAssets/Private/Misc/MediaTextureResource.cpp
--- a/Engine/Source/Runtime/MediaAssets/Private/Misc/MediaTextureResource.cpp
+++ b/Engine/Source/Runtime/MediaAssets/Private/Misc/MediaTextureResource.cpp
@@ -11,9 +11,9 @@
 FMediaTextureResource::FMediaTextureResource(UMediaTexture& InOwner, const FLinearColor& InClearColor, const FIntPoint& InOutputDimensions, EMediaTextureSinkFormat InSinkFormat, EMediaTextureSinkMode InSinkMode)
 	: Owner(InOwner)
 	, TripleBuffer(BufferResources)
-	, BufferBytesPerPixel(0)
 	, BufferClearColor(InClearColor)
-	, BufferDimensions(FIntPoint::ZeroValue)
+	, BufferDimensions(InOutputDimensions)
+	, BufferPitch(0)
 	, CachedResourceSizeBytes(0)
 	, OutputDimensions(InOutputDimensions)
 	, RequiresConversion(false)
@@ -89,18 +89,18 @@ FRHITexture* FMediaTextureResource::GetTexture()
 }
 
 
-void FMediaTextureResource::InitializeBuffer(FIntPoint Dimensions, EMediaTextureSinkFormat Format, EMediaTextureSinkMode Mode)
+void FMediaTextureResource::InitializeBuffer(FIntPoint OutputDim, FIntPoint BufferDim, EMediaTextureSinkFormat Format, EMediaTextureSinkMode Mode)
 {
 	State = EState::Initializing;
 
 	if (IsInActualRenderingThread())
 	{
-		InitializeResource(Dimensions, Format, Mode);
+		InitializeResource(OutputDim, BufferDim, Format, Mode);
 	}
 	else
 	{
 		RenderThreadTasks.Enqueue([=]() {
-			InitializeResource(Dimensions, Format, Mode);
+			InitializeResource(OutputDim, BufferDim, Format, Mode);
 		});
 	}
 }
@@ -152,21 +152,13 @@ void FMediaTextureResource::ShutdownBuffer()
 }
 
 
-void FMediaTextureResource::UpdateBuffer(const uint8* Data, uint32 Pitch)
+void FMediaTextureResource::UpdateBuffer(const uint8* Data, uint32 /*Pitch*/)
 {
 	if (State != EState::Initialized)
 	{
 		return;
 	}
 
-	// calculate pitch
-	const SIZE_T BytesPerRow = OutputDimensions.X * BufferBytesPerPixel;
-
-	if (Pitch == 0)
-	{
-		Pitch = BytesPerRow;
-	}
-
 	// update sink
 	if (SinkMode == EMediaTextureSinkMode::Buffered)
 	{
@@ -174,21 +166,7 @@ void FMediaTextureResource::UpdateBuffer(const uint8* Data, uint32 Pitch)
 
 		if (Dest != nullptr)
 		{
-			if (Pitch != BytesPerRow)
-			{
-				const uint8* Row = Data;
-
-				for (int32 RowIndex = 0; RowIndex < BufferDimensions.Y; ++RowIndex)
-				{
-					FMemory::Memcpy(Dest, Row, BytesPerRow);
-					Dest += BytesPerRow;
-					Row += Pitch;
-				}
-			}
-			else
-			{
-				FMemory::Memcpy(Dest, Data, BufferDimensions.Y * BytesPerRow);
-			}
+			FMemory::Memcpy(Dest, Data, BufferPitch * BufferDimensions.Y);
 
 			if (TripleBuffer.IsDirty())
 			{
@@ -203,13 +181,13 @@ void FMediaTextureResource::UpdateBuffer(const uint8* Data, uint32 Pitch)
 		if (RequiresConversion)
 		{
 			FUpdateTextureRegion2D Region(0, 0, 0, 0, BufferDimensions.X, BufferDimensions.Y);
-			RHIUpdateTexture2D(BufferResources[2].RenderTarget.GetReference(), 0, Region, Pitch, Data);
+			RHIUpdateTexture2D(BufferResources[2].RenderTarget.GetReference(), 0, Region, BufferPitch, Data);
 			ConvertResource(BufferResources[2]);
 		}
 		else
 		{
 			FUpdateTextureRegion2D Region(0, 0, 0, 0, OutputDimensions.X, OutputDimensions.Y);
-			RHIUpdateTexture2D(RenderTargetTextureRHI.GetReference(), 0, Region, Pitch, Data);
+			RHIUpdateTexture2D(RenderTargetTextureRHI.GetReference(), 0, Region, BufferPitch, Data);
 		}
 	}
 	else
@@ -289,63 +267,68 @@ void FMediaTextureResource::InitDynamicRHI()
 
 	SamplerStateRHI = RHICreateSamplerState(SamplerStateInitializer);
 
-	if (OutputDimensions.GetMin() <= 0)
+	if ((OutputDimensions.GetMin() <= 0) || (BufferDimensions.GetMin() <= 0))
 	{
 		return;
 	}
 
 	// determine buffer resource pixel format
 	EPixelFormat BufferFormat;
+	bool SrgbBuffer = false;
+
+	switch (SinkFormat)
 	{
-		switch (SinkFormat)
-		{
-		case EMediaTextureSinkFormat::CharNV12:
-		case EMediaTextureSinkFormat::CharNV21:
-			BufferBytesPerPixel = 1;
-			BufferDimensions = FIntPoint(OutputDimensions.X, OutputDimensions.Y * 1.5f);
-			BufferFormat = PF_G8;
-			RequiresConversion = true;
-			break;
-
-		case EMediaTextureSinkFormat::CharUYVY:
-		case EMediaTextureSinkFormat::CharYUY2:
-		case EMediaTextureSinkFormat::CharYVYU:
-			if ((OutputDimensions.X & 0x1) != 0)
-			{
-				return;
-			}
+	case EMediaTextureSinkFormat::CharNV12:
+	case EMediaTextureSinkFormat::CharNV21:
+		BufferFormat = PF_G8;
+		BufferPitch = BufferDimensions.X;
+		RequiresConversion = true;
+		break;
 
-			BufferBytesPerPixel = 2;
-			BufferDimensions = FIntPoint(OutputDimensions.X / 2, OutputDimensions.Y);
-			BufferFormat = PF_B8G8R8A8;
-			RequiresConversion = true;
-			break;
-
-		case EMediaTextureSinkFormat::CharAYUV:
-		case EMediaTextureSinkFormat::CharBGRA:
-			BufferBytesPerPixel = 4;
-			BufferDimensions = OutputDimensions;
-			BufferFormat = PF_B8G8R8A8;
-			RequiresConversion = false;
-			break;
-
-		case EMediaTextureSinkFormat::FloatRGB:
-			BufferBytesPerPixel = 12;
-			BufferDimensions = OutputDimensions;
-			BufferFormat = PF_FloatRGB;
-			RequiresConversion = false;
-			break;
-
-		case EMediaTextureSinkFormat::FloatRGBA:
-			BufferBytesPerPixel = 16;
-			BufferDimensions = OutputDimensions;
-			BufferFormat = PF_FloatRGBA;
-			RequiresConversion = false;
-			break;
-
-		default:
-			return;
-		}
+	case EMediaTextureSinkFormat::CharUYVY:
+	case EMediaTextureSinkFormat::CharYUY2:
+	case EMediaTextureSinkFormat::CharYVYU:
+		BufferFormat = PF_B8G8R8A8;
+		BufferPitch = BufferDimensions.X * 4;
+		RequiresConversion = true;
+		break;
+
+	case EMediaTextureSinkFormat::CharAYUV:
+		BufferFormat = PF_B8G8R8A8;
+		BufferPitch = BufferDimensions.X * 4;
+		RequiresConversion = true;
+		break;
+
+	case EMediaTextureSinkFormat::CharBGRA:
+		BufferFormat = PF_B8G8R8A8;
+		BufferPitch = BufferDimensions.X * 4;
+		RequiresConversion = (BufferDimensions != OutputDimensions);
+		SrgbBuffer = true;
+		break;
+
+	case EMediaTextureSinkFormat::CharBMP:
+		BufferFormat = PF_B8G8R8A8;
+		BufferPitch = BufferDimensions.X * 4;
+		RequiresConversion = true;
+		SrgbBuffer = true;
+		break;
+
+	case EMediaTextureSinkFormat::FloatRGB:
+		BufferFormat = PF_FloatRGB;
+		BufferPitch = BufferDimensions.X * 12;
+		RequiresConversion = (BufferDimensions != OutputDimensions);
+		SrgbBuffer = true;
+		break;
+
+	case EMediaTextureSinkFormat::FloatRGBA:
+		BufferFormat = PF_FloatRGBA;
+		BufferPitch = BufferDimensions.X * 16;
+		RequiresConversion = (BufferDimensions != OutputDimensions);
+		SrgbBuffer = true;
+		break;
+
+	default:
+		return;
 	}
 
 	// determine how many buffer resources we need
@@ -382,7 +365,7 @@ void FMediaTextureResource::InitDynamicRHI()
 				BufferDimensions.Y,
 				BufferFormat,
 				1,
-				TexCreate_Dynamic | TexCreate_NoTiling,
+				TexCreate_Dynamic | TexCreate_NoTiling | (SrgbBuffer ? TexCreate_SRGB : 0),
 				TexCreate_RenderTargetable,
 				false,
 				CreateInfo,
@@ -534,7 +517,7 @@ void FMediaTextureResource::CacheResourceSize()
 
 	if (OutputDimensions.GetMin() > 0)
 	{
-		const SIZE_T BufferSize = BufferDimensions.X * BufferDimensions.Y * BufferBytesPerPixel;
+		const SIZE_T BufferSize = BufferPitch * BufferDimensions.Y;
 
 		if (SinkMode == EMediaTextureSinkMode::Buffered)
 		{
@@ -583,7 +566,7 @@ void FMediaTextureResource::ConvertResource(const FResource& Resource)
 			static FGlobalBoundShaderState BoundShaderState;
 
 			SetGlobalBoundShaderState(RHICmdList, GMaxRHIFeatureLevel, BoundShaderState, GMediaVertexDeclaration.VertexDeclarationRHI, *VertexShader, *ConvertShader);
-			ConvertShader->SetParameters(RHICmdList, Resource.ShaderResource);
+			ConvertShader->SetParameters(RHICmdList, Resource.ShaderResource, OutputDimensions);
 		}
 		break;
 
@@ -593,7 +576,7 @@ void FMediaTextureResource::ConvertResource(const FResource& Resource)
 			static FGlobalBoundShaderState BoundShaderState;
 
 			SetGlobalBoundShaderState(RHICmdList, GMaxRHIFeatureLevel, BoundShaderState, GMediaVertexDeclaration.VertexDeclarationRHI, *VertexShader, *ConvertShader);
-			ConvertShader->SetParameters(RHICmdList, Resource.ShaderResource);
+			ConvertShader->SetParameters(RHICmdList, Resource.ShaderResource, OutputDimensions);
 		}
 		break;
 
@@ -613,7 +596,7 @@ void FMediaTextureResource::ConvertResource(const FResource& Resource)
 			static FGlobalBoundShaderState BoundShaderState;
 
 			SetGlobalBoundShaderState(RHICmdList, GMaxRHIFeatureLevel, BoundShaderState, GMediaVertexDeclaration.VertexDeclarationRHI, *VertexShader, *ConvertShader);
-			ConvertShader->SetParameters(RHICmdList, Resource.ShaderResource);
+			ConvertShader->SetParameters(RHICmdList, Resource.ShaderResource, OutputDimensions);
 		}
 		break;
 
@@ -627,6 +610,28 @@ void FMediaTextureResource::ConvertResource(const FResource& Resource)
 		}
 		break;
 
+	case EMediaTextureSinkFormat::CharBGRA:
+	case EMediaTextureSinkFormat::FloatRGB:
+	case EMediaTextureSinkFormat::FloatRGBA:
+		{
+			TShaderMapRef<FRGBConvertPS> ConvertShader(ShaderMap);
+			static FGlobalBoundShaderState BoundShaderState;
+
+			SetGlobalBoundShaderState(RHICmdList, GMaxRHIFeatureLevel, BoundShaderState, GMediaVertexDeclaration.VertexDeclarationRHI, *VertexShader, *ConvertShader);
+			ConvertShader->SetParameters(RHICmdList, Resource.ShaderResource, OutputDimensions);
+		}
+		break;
+
+	case EMediaTextureSinkFormat::CharBMP:
+		{
+			TShaderMapRef<FBMPConvertPS> ConvertShader(ShaderMap);
+			static FGlobalBoundShaderState BoundShaderState;
+
+			SetGlobalBoundShaderState(RHICmdList, GMaxRHIFeatureLevel, BoundShaderState, GMediaVertexDeclaration.VertexDeclarationRHI, *VertexShader, *ConvertShader);
+			ConvertShader->SetParameters(RHICmdList, Resource.ShaderResource, OutputDimensions);
+		}
+		break;
+
 	default:
 		return;
 	}
@@ -668,11 +673,11 @@ void FMediaTextureResource::DisplayResource(const FResource& Resource)
 }
 
 
-void FMediaTextureResource::InitializeResource(FIntPoint Dimensions, EMediaTextureSinkFormat Format, EMediaTextureSinkMode Mode)
+void FMediaTextureResource::InitializeResource(FIntPoint OutputDim, FIntPoint BufferDim, EMediaTextureSinkFormat Format, EMediaTextureSinkMode Mode)
 {
 	check(IsInRenderingThread());
 
-	if ((Dimensions == OutputDimensions) && (Format == SinkFormat) && (Mode == SinkMode))
+	if ((OutputDim == OutputDimensions) && (BufferDim == BufferDimensions) && (Format == SinkFormat) && (Mode == SinkMode))
 	{
 		// reuse existing resources
 		State = (OutputDimensions.GetMin() > 0) ? EState::Initialized : EState::ShutDown;
@@ -682,7 +687,8 @@ void FMediaTextureResource::InitializeResource(FIntPoint Dimensions, EMediaTextu
 	{
 		// reinitialize resources
 		ReleaseDynamicRHI();
-		OutputDimensions = Dimensions;
+		BufferDimensions = BufferDim;
+		OutputDimensions = OutputDim;
 		SinkFormat = Format;
 		SinkMode = Mode;
 		InitDynamicRHI();
diff --git a/Engine/Source/Runtime/MediaAssets/Private/Misc/MediaTextureResource.h b/Engine/Source/Runtime/MediaAssets/Private/Misc/MediaTextureResource.h
--- a/Engine/Source/Runtime/MediaAssets/Private/Misc/MediaTextureResource.h
+++ b/Engine/Source/Runtime/MediaAssets/Private/Misc/MediaTextureResource.h
@@ -86,11 +86,12 @@ public:
 	/**
 	 * Initialize the render target buffer(s).
 	 *
-	 * @param Dimensions Width and height of the texture (in pixels).
+	 * @param OutputDim Width and height of the video output (in pixels).
+	 * @param BufferDim Width and height of the sink buffer(s) (in pixels).
 	 * @param Format The pixel format of the sink's render target texture.
 	 * @param Mode The mode to operate the sink in (buffered vs. unbuffered).
 	 */
-	void InitializeBuffer(FIntPoint Dimensions, EMediaTextureSinkFormat Format, EMediaTextureSinkMode Mode);
+	void InitializeBuffer(FIntPoint OutputDim, FIntPoint BufferDim, EMediaTextureSinkFormat Format, EMediaTextureSinkMode Mode);
 
 	/** Release a previously acquired texture buffer. */
 	void ReleaseBuffer();
@@ -160,11 +161,12 @@ protected:
 	/**
 	 * Initialize this resource.
 	 *
-	 * @param Dimensions The new texture dimensions.
+	 * @param OutputDim Width and height of the output texture (in pixels).
+	 * @param BufferDim Width and height of the buffer texture(s) (in pixels).
 	 * @param Format The new texture format.
 	 * @param Mode The new sink mode.
 	 */
-	void InitializeResource(FIntPoint Dimensions, EMediaTextureSinkFormat Format, EMediaTextureSinkMode Mode);
+	void InitializeResource(FIntPoint OutputDim, FIntPoint BufferDim, EMediaTextureSinkFormat Format, EMediaTextureSinkMode Mode);
 
 	/** Process any queued up tasks on the render thread. */
 	void ProcessRenderThreadTasks();
@@ -191,9 +193,6 @@ private:
 
 	//~ The following fields are owned by the render thread
 
-	/** Number of bytes per pixel in buffer resources. */
-	uint8 BufferBytesPerPixel;
-
 	/** The clear color to use. */
 	FLinearColor BufferClearColor;
 
@@ -205,6 +204,9 @@ private:
 	 */
 	FIntPoint BufferDimensions;
 
+	/** Number of bytes per row in buffer resources. */
+	SIZE_T BufferPitch;
+
 	/**
 	 * Texture resources for buffered mode or pixel conversions.
 	 *
diff --git a/Engine/Source/Runtime/MediaAssets/Public/MediaPlayer.h b/Engine/Source/Runtime/MediaAssets/Public/MediaPlayer.h
--- a/Engine/Source/Runtime/MediaAssets/Public/MediaPlayer.h
+++ b/Engine/Source/Runtime/MediaAssets/Public/MediaPlayer.h
@@ -462,12 +462,12 @@ public:
 	/**
 	 * Seeks to the specified playback time.
 	 *
-	 * @param InTime The playback time to set.
+	 * @param Time The playback time to set.
 	 * @return true on success, false otherwise.
 	 * @see GetTime, Rewind
 	 */
 	UFUNCTION(BlueprintCallable, Category="Media|MediaPlayer")
-	bool Seek(const FTimespan& InTime);
+	bool Seek(const FTimespan& Time);
 
 	/**
 	 * Select the active track of the given type.
@@ -490,7 +490,7 @@ public:
 	 * @see IsLooping
 	 */
 	UFUNCTION(BlueprintCallable, Category="Media|MediaPlayer")
-	bool SetLooping(bool InLooping);
+	bool SetLooping(bool Looping);
 
 	/**
 	 * Changes the media's playback rate.
diff --git a/Engine/Source/Runtime/MediaAssets/Public/MediaTexture.h b/Engine/Source/Runtime/MediaAssets/Public/MediaTexture.h
--- a/Engine/Source/Runtime/MediaAssets/Public/MediaTexture.h
+++ b/Engine/Source/Runtime/MediaAssets/Public/MediaTexture.h
@@ -21,17 +21,46 @@ class MEDIAASSETS_API UMediaTexture
 	GENERATED_UCLASS_BODY()
 
 	/** The addressing mode to use for the X axis. */
-	UPROPERTY(EditAnywhere, BlueprintReadWrite, Category=MediaTexture, AssetRegistrySearchable)
+	UPROPERTY(EditAnywhere, BlueprintReadWrite, Category="Media|MediaTexture", AssetRegistrySearchable)
 	TEnumAsByte<TextureAddress> AddressX;
 
 	/** The addressing mode to use for the Y axis. */
-	UPROPERTY(EditAnywhere, BlueprintReadWrite, Category=MediaTexture, AssetRegistrySearchable)
+	UPROPERTY(EditAnywhere, BlueprintReadWrite, Category="Media|MediaTexture", AssetRegistrySearchable)
 	TEnumAsByte<TextureAddress> AddressY;
 
 	/** The color used to clear the texture if CloseAction is set to Clear (default = black). */
-	UPROPERTY(EditAnywhere, BlueprintReadWrite, Category=MediaTexture, AdvancedDisplay)
+	UPROPERTY(EditAnywhere, BlueprintReadWrite, Category="Media|MediaTexture", AdvancedDisplay)
 	FLinearColor ClearColor;
 
+public:
+
+	/**
+	 * Gets the current aspect ratio of the texture.
+	 *
+	 * @return Texture aspect ratio.
+	 * @see GetHeight, GetWidth
+	 */
+	UFUNCTION(BlueprintCallable, Category="Media|MediaTexture")
+	float GetAspectRatio() const;
+
+	/**
+	 * Gets the current height of the texture.
+	 *
+	 * @return Texture height (in pixels).
+	 * @see GetAspectRatio, GetWidth
+	 */
+	UFUNCTION(BlueprintCallable, Category="Media|MediaTexture")
+	int32 GetHeight() const;
+
+	/**
+	 * Gets the current width of the texture.
+	 *
+	 * @return Texture width (in pixels).
+	 * @see GetAspectRatio, GetHeight
+	 */
+	UFUNCTION(BlueprintCallable, Category="Media|MediaTexture")
+	int32 GetWidth() const;
+
 public:
 
 	DECLARE_EVENT_OneParam(UMediaTexture, FOnBeginDestroy, UMediaTexture& /*DestroyedMediaTexture*/)
@@ -50,7 +79,7 @@ public:
 	virtual EMediaTextureSinkFormat GetTextureSinkFormat() const override;
 	virtual EMediaTextureSinkMode GetTextureSinkMode() const override;
 	virtual FRHITexture* GetTextureSinkTexture() override;
-	virtual bool InitializeTextureSink(FIntPoint Dimensions, EMediaTextureSinkFormat Format, EMediaTextureSinkMode Mode) override;
+	virtual bool InitializeTextureSink(FIntPoint OutputDim, FIntPoint BufferDim, EMediaTextureSinkFormat Format, EMediaTextureSinkMode Mode) override;
 	virtual void ReleaseTextureSinkBuffer() override;
 	virtual void ShutdownTextureSink() override;
 	virtual bool SupportsTextureSinkFormat(EMediaTextureSinkFormat Format) const override;
@@ -100,12 +129,15 @@ private:
 	/** Critical section for synchronizing access to texture resource object. */
 	mutable FCriticalSection CriticalSection;
 
-	/** Width and height of the texture (in pixels). */
-	FIntPoint SinkDimensions;
+	/** Width and height of the sink buffers (in pixels). */
+	FIntPoint SinkBufferDim;
+
+	/** Width and height of the video output (in pixels). */
+	FIntPoint SinkOutputDim;
 
 	/** The render target's pixel format. */
 	EMediaTextureSinkFormat SinkFormat;
 
 	/** The mode that this sink is currently operating in. */
 	EMediaTextureSinkMode SinkMode;
-	};
+};
diff --git a/Engine/Source/Runtime/Messaging/Private/Bridge/MessageAddressBook.h b/Engine/Source/Runtime/Messaging/Private/Bridge/MessageAddressBook.h
--- a/Engine/Source/Runtime/Messaging/Private/Bridge/MessageAddressBook.h
+++ b/Engine/Source/Runtime/Messaging/Private/Bridge/MessageAddressBook.h
@@ -2,6 +2,8 @@
 
 #pragma once
 
+#include "IMessageContext.h"
+
 
 /**
  * Implements an address book that maps message addresses to remote nodes.
diff --git a/Engine/Source/Runtime/Messaging/Private/Bridge/MessageBridge.cpp b/Engine/Source/Runtime/Messaging/Private/Bridge/MessageBridge.cpp
--- a/Engine/Source/Runtime/Messaging/Private/Bridge/MessageBridge.cpp
+++ b/Engine/Source/Runtime/Messaging/Private/Bridge/MessageBridge.cpp
@@ -1,12 +1,17 @@
 // Copyright 1998-2016 Epic Games, Inc. All Rights Reserved.
 
 #include "MessagingPrivatePCH.h"
+#include "MessageBridge.h"
 
 
 /* FMessageBridge structors
  *****************************************************************************/
 
-FMessageBridge::FMessageBridge(const FMessageAddress InAddress, const IMessageBusRef& InBus, const IMessageTransportRef& InTransport)
+FMessageBridge::FMessageBridge(
+	const FMessageAddress InAddress,
+	const TSharedRef<IMessageBus, ESPMode::ThreadSafe>& InBus,
+	const TSharedRef<IMessageTransport, ESPMode::ThreadSafe>& InTransport
+)
 	: Address(InAddress)
 	, Bus(InBus)
 	, Enabled(false)
@@ -123,7 +128,7 @@ void FMessageBridge::ReceiveMessage(const IMessageContextRef& Context)
 /* ISendMessages interface
  *****************************************************************************/
 
-void FMessageBridge::NotifyMessageError(const IMessageContextRef& Context, const FString& Error)
+void FMessageBridge::NotifyMessageError(const TSharedRef<IMessageContext, ESPMode::ThreadSafe>& Context, const FString& Error)
 {
 	// deprecated
 }
@@ -154,7 +159,7 @@ void FMessageBridge::HandleMessageBusShutdown()
 }
 
 
-void FMessageBridge::HandleTransportMessageReceived(const IMessageContextRef& Context, const FGuid& NodeId)
+void FMessageBridge::HandleTransportMessageReceived(const TSharedRef<IMessageContext, ESPMode::ThreadSafe>& Context, const FGuid& NodeId)
 {
 	if (!Enabled || !Bus.IsValid())
 	{
diff --git a/Engine/Source/Runtime/Messaging/Private/Bridge/MessageBridge.h b/Engine/Source/Runtime/Messaging/Private/Bridge/MessageBridge.h
--- a/Engine/Source/Runtime/Messaging/Private/Bridge/MessageBridge.h
+++ b/Engine/Source/Runtime/Messaging/Private/Bridge/MessageBridge.h
@@ -4,9 +4,14 @@
 
 #include "IMessageBridge.h"
 #include "IMessageContext.h"
-#include "IMessageTransport.h"
-#include "IReceiveMessages.h"
-#include "ISendMessages.h"
+#include "IMessageReceiver.h"
+#include "IMessageSender.h"
+#include "MessageAddressBook.h"
+
+
+class IMessageBus;
+class IMessageSubscription;
+class IMessageTransport;
 
 
 /**
@@ -26,8 +31,8 @@
 class FMessageBridge
 	: public TSharedFromThis<FMessageBridge, ESPMode::ThreadSafe>
 	, public IMessageBridge
-	, public IReceiveMessages
-	, public ISendMessages
+	, public IMessageReceiver
+	, public IMessageSender
 {
 public:
 
@@ -38,14 +43,17 @@ public:
 	 * @param InBus The message bus that this node is connected to.
 	 * @param InTransport The transport mechanism to use.
 	 */
-	FMessageBridge(const FMessageAddress InAddress, const IMessageBusRef& InBus, const IMessageTransportRef& InTransport);
+	FMessageBridge(
+		const FMessageAddress InAddress,
+		const TSharedRef<IMessageBus, ESPMode::ThreadSafe>& InBus,
+		const TSharedRef<IMessageTransport, ESPMode::ThreadSafe>& InTransport);
 
 	/** Virtual destructor. */
 	virtual ~FMessageBridge();
 
 public:
 
-	// IMessageBridge interface
+	//~ IMessageBridge interface
 
 	virtual void Disable() override;
 	virtual void Enable() override;
@@ -57,7 +65,7 @@ public:
 
 public:
 
-	// IReceiveMessages interface
+	//~ IReceiveMessages interface
 
 	virtual FName GetDebugName() const override
 	{
@@ -83,14 +91,14 @@ public:
 
 public:
 
-	// ISendMessages interface
+	//~ ISendMessages interface
 
 	virtual FMessageAddress GetSenderAddress() override
 	{
 		return Address;
 	}
 
-	virtual void NotifyMessageError(const IMessageContextRef& Context, const FString& Error) override;
+	virtual void NotifyMessageError(const TSharedRef<IMessageContext, ESPMode::ThreadSafe>& Context, const FString& Error) override;
 
 protected:
 
@@ -103,7 +111,7 @@ private:
 	void HandleMessageBusShutdown();
 
 	/** Callback for messages received from the transport layer. */
-	void HandleTransportMessageReceived(const IMessageContextRef& Envelope, const FGuid& NodeId);
+	void HandleTransportMessageReceived(const TSharedRef<IMessageContext, ESPMode::ThreadSafe>& Envelope, const FGuid& NodeId);
 
 	/** Callback for lost remote nodes. */
 	void HandleTransportNodeLost(const FGuid& NodeId);
@@ -117,7 +125,7 @@ private:
 	FMessageAddressBook AddressBook;
 
 	/** Holds a reference to the bus that this bridge is attached to. */
-	IMessageBusPtr Bus;
+	TSharedPtr<IMessageBus, ESPMode::ThreadSafe> Bus;
 
 	/** Hold a flag indicating whether this endpoint is active. */
 	bool Enabled;
@@ -126,8 +134,8 @@ private:
 	const FGuid Id;
 
 	/** Holds the message subscription for outbound messages. */
-	IMessageSubscriptionPtr MessageSubscription;
+	TSharedPtr<IMessageSubscription, ESPMode::ThreadSafe> MessageSubscription;
 
 	/** Holds the message transport object. */
-	IMessageTransportPtr Transport;
+	TSharedPtr<IMessageTransport, ESPMode::ThreadSafe> Transport;
 };
diff --git a/Engine/Source/Runtime/Messaging/Private/Bus/MessageBus.cpp b/Engine/Source/Runtime/Messaging/Private/Bus/MessageBus.cpp
--- a/Engine/Source/Runtime/Messaging/Private/Bus/MessageBus.cpp
+++ b/Engine/Source/Runtime/Messaging/Private/Bus/MessageBus.cpp
@@ -1,12 +1,16 @@
 // Copyright 1998-2016 Epic Games, Inc. All Rights Reserved.
 
 #include "MessagingPrivatePCH.h"
+#include "MessageBus.h"
+#include "MessageContext.h"
+#include "MessageRouter.h"
+#include "MessageSubscription.h"
 
 
 /* FMessageBus structors
  *****************************************************************************/
 
-FMessageBus::FMessageBus(const IAuthorizeMessageRecipientsPtr& InRecipientAuthorizer)
+FMessageBus::FMessageBus(const TSharedPtr<IAuthorizeMessageRecipients>& InRecipientAuthorizer)
 	: RecipientAuthorizer(InRecipientAuthorizer)
 {
 	Router = new FMessageRouter();
@@ -27,19 +31,31 @@ FMessageBus::~FMessageBus()
 /* IMessageBus interface
  *****************************************************************************/
 
-void FMessageBus::Forward(const IMessageContextRef& Context, const TArray<FMessageAddress>& Recipients, const FTimespan& Delay, const ISendMessagesRef& Forwarder)
+void FMessageBus::Forward(
+	const TSharedRef<IMessageContext, ESPMode::ThreadSafe>& Context,
+	const TArray<FMessageAddress>& Recipients,
+	const FTimespan& Delay,
+	const TSharedRef<IMessageSender, ESPMode::ThreadSafe>& Forwarder
+)
 {
-	Router->RouteMessage(MakeShareable(new FMessageContext(Context, Forwarder->GetSenderAddress(), Recipients, EMessageScope::Process, FDateTime::UtcNow() + Delay, FTaskGraphInterface::Get().GetCurrentThreadIfKnown())));
+	Router->RouteMessage(MakeShareable(new FMessageContext(
+		Context,
+		Forwarder->GetSenderAddress(),
+		Recipients,
+		EMessageScope::Process,
+		FDateTime::UtcNow() + Delay,
+		FTaskGraphInterface::Get().GetCurrentThreadIfKnown()
+	)));
 }
 
 
-IMessageTracerRef FMessageBus::GetTracer()
+TSharedRef<IMessageTracer, ESPMode::ThreadSafe> FMessageBus::GetTracer()
 {
 	return Router->GetTracer();
 }
 
 
-void FMessageBus::Intercept(const IMessageInterceptorRef& Interceptor, const FName& MessageType)
+void FMessageBus::Intercept(const TSharedRef<IMessageInterceptor, ESPMode::ThreadSafe>& Interceptor, const FName& MessageType)
 {
 	if (MessageType == NAME_None)
 	{
@@ -59,21 +75,56 @@ FOnMessageBusShutdown& FMessageBus::OnShutdown()
 }
 
 
-void FMessageBus::Publish(void* Message, UScriptStruct* TypeInfo, EMessageScope Scope, const FTimespan& Delay, const FDateTime& Expiration, const ISendMessagesRef& Publisher)
+void FMessageBus::Publish(
+	void* Message,
+	UScriptStruct* TypeInfo,
+	EMessageScope Scope,
+	const FTimespan& Delay,
+	const FDateTime& Expiration,
+	const TSharedRef<IMessageSender, ESPMode::ThreadSafe>& Publisher
+)
 {
-	Router->RouteMessage(MakeShareable(new FMessageContext(Message, TypeInfo, nullptr, Publisher->GetSenderAddress(), TArray<FMessageAddress>(), Scope, FDateTime::UtcNow() + Delay, Expiration, FTaskGraphInterface::Get().GetCurrentThreadIfKnown())));
+	Router->RouteMessage(MakeShareable(new FMessageContext(
+		Message,
+		TypeInfo,
+		nullptr,
+		Publisher->GetSenderAddress(),
+		TArray<FMessageAddress>(),
+		Scope,
+		FDateTime::UtcNow() + Delay,
+		Expiration,
+		FTaskGraphInterface::Get().GetCurrentThreadIfKnown()
+	)));
 }
 
 
-void FMessageBus::Register(const FMessageAddress& Address, const IReceiveMessagesRef& Recipient)
+void FMessageBus::Register(const FMessageAddress& Address, const TSharedRef<IMessageReceiver, ESPMode::ThreadSafe>& Recipient)
 {
 	Router->AddRecipient(Address, Recipient);
 }
 
 
-void FMessageBus::Send(void* Message, UScriptStruct* TypeInfo, const IMessageAttachmentPtr& Attachment, const TArray<FMessageAddress>& Recipients, const FTimespan& Delay, const FDateTime& Expiration, const ISendMessagesRef& Sender)
+void FMessageBus::Send(
+	void* Message,
+	UScriptStruct* TypeInfo,
+	const TSharedPtr<IMessageAttachment, ESPMode::ThreadSafe>& Attachment,
+	const TArray<FMessageAddress>& Recipients,
+	const FTimespan& Delay,
+	const FDateTime& Expiration,
+	const TSharedRef<IMessageSender, ESPMode::ThreadSafe>& Sender
+)
 {
-	Router->RouteMessage(MakeShareable(new FMessageContext(Message, TypeInfo, Attachment, Sender->GetSenderAddress(), Recipients, EMessageScope::Network, FDateTime::UtcNow() + Delay, Expiration, FTaskGraphInterface::Get().GetCurrentThreadIfKnown())));
+	Router->RouteMessage(MakeShareable(new FMessageContext(
+		Message,
+		TypeInfo,
+		Attachment,
+		Sender->GetSenderAddress(),
+		Recipients,
+		EMessageScope::Network,
+		FDateTime::UtcNow() + Delay,
+		Expiration,
+		FTaskGraphInterface::Get().GetCurrentThreadIfKnown()
+	)));
 }
 
 
@@ -90,13 +141,17 @@ void FMessageBus::Shutdown()
 }
 
 
-IMessageSubscriptionPtr FMessageBus::Subscribe(const IReceiveMessagesRef& Subscriber, const FName& MessageType, const FMessageScopeRange& ScopeRange)
+TSharedPtr<IMessageSubscription, ESPMode::ThreadSafe> FMessageBus::Subscribe(
+	const TSharedRef<IMessageReceiver, ESPMode::ThreadSafe>& Subscriber,
+	const FName& MessageType,
+	const FMessageScopeRange& ScopeRange
+)
 {
 	if (MessageType != NAME_None)
 	{
 		if (!RecipientAuthorizer.IsValid() || RecipientAuthorizer->AuthorizeSubscription(Subscriber, MessageType))
 		{
-			IMessageSubscriptionRef Subscription = MakeShareable(new FMessageSubscription(Subscriber, MessageType, ScopeRange));
+			TSharedRef<IMessageSubscription, ESPMode::ThreadSafe> Subscription = MakeShareable(new FMessageSubscription(Subscriber, MessageType, ScopeRange));
 			Router->AddSubscription(Subscription);
 
 			return Subscription;
@@ -107,7 +162,7 @@ IMessageSubscriptionPtr FMessageBus::Subscribe(const IReceiveMessagesRef& Subscr
 }
 
 
-void FMessageBus::Unintercept(const IMessageInterceptorRef& Interceptor, const FName& MessageType)
+void FMessageBus::Unintercept(const TSharedRef<IMessageInterceptor, ESPMode::ThreadSafe>& Interceptor, const FName& MessageType)
 {
 	if (MessageType != NAME_None)
 	{
@@ -125,7 +180,7 @@ void FMessageBus::Unregister(const FMessageAddress& Address)
 }
 
 
-void FMessageBus::Unsubscribe(const IReceiveMessagesRef& Subscriber, const FName& MessageType)
+void FMessageBus::Unsubscribe(const TSharedRef<IMessageReceiver, ESPMode::ThreadSafe>& Subscriber, const FName& MessageType)
 {
 	if (MessageType != NAME_None)
 	{
diff --git a/Engine/Source/Runtime/Messaging/Private/Bus/MessageBus.h b/Engine/Source/Runtime/Messaging/Private/Bus/MessageBus.h
--- a/Engine/Source/Runtime/Messaging/Private/Bus/MessageBus.h
+++ b/Engine/Source/Runtime/Messaging/Private/Bus/MessageBus.h
@@ -3,18 +3,17 @@
 #pragma once
 
 #include "IAuthorizeMessageRecipients.h"
-#include "IMessageAttachment.h"
 #include "IMessageBus.h"
 #include "IMessageContext.h"
-#include "IMessageInterceptor.h"
-#include "IMessageSubscription.h"
-#include "IReceiveMessages.h"
-#include "ISendMessages.h"
 
 
-// forward declarations
 class FMessageRouter;
 class FRunnableThread;
+class IMessageInterceptor;
+class IMessageSubscription;
+class IMessageReceiver;
+class IMessageSender;
+class IMessageTracer;
 
 
 /**
@@ -31,27 +30,27 @@ public:
 	 *
 	 * @param InRecipientAuthorizer An optional recipient authorizer.
 	 */
-	FMessageBus(const IAuthorizeMessageRecipientsPtr& InRecipientAuthorizer);
+	FMessageBus(const TSharedPtr<IAuthorizeMessageRecipients>& InRecipientAuthorizer);
 
 	/** Virtual destructor. */
 	virtual ~FMessageBus();
 
 public:
 
-	// IMessageBus interface
+	//~ IMessageBus interface
 
-	virtual void Forward(const IMessageContextRef& Context, const TArray<FMessageAddress>& Recipients, const FTimespan& Delay, const ISendMessagesRef& Forwarder) override;
-	virtual IMessageTracerRef GetTracer() override;
-	virtual void Intercept(const IMessageInterceptorRef& Interceptor, const FName& MessageType) override;
+	virtual void Forward(const TSharedRef<IMessageContext, ESPMode::ThreadSafe>& Context, const TArray<FMessageAddress>& Recipients, const FTimespan& Delay, const TSharedRef<IMessageSender, ESPMode::ThreadSafe>& Forwarder) override;
+	virtual TSharedRef<IMessageTracer, ESPMode::ThreadSafe> GetTracer() override;
+	virtual void Intercept(const TSharedRef<IMessageInterceptor, ESPMode::ThreadSafe>& Interceptor, const FName& MessageType) override;
 	virtual FOnMessageBusShutdown& OnShutdown() override;
-	virtual void Publish(void* Message, UScriptStruct* TypeInfo, EMessageScope Scope, const FTimespan& Delay, const FDateTime& Expiration, const ISendMessagesRef& Publisher) override;
-	virtual void Register(const FMessageAddress& Address, const IReceiveMessagesRef& Recipient) override;
-	virtual void Send(void* Message, UScriptStruct* TypeInfo, const IMessageAttachmentPtr& Attachment, const TArray<FMessageAddress>& Recipients, const FTimespan& Delay, const FDateTime& Expiration, const ISendMessagesRef& Sender) override;
+	virtual void Publish(void* Message, UScriptStruct* TypeInfo, EMessageScope Scope, const FTimespan& Delay, const FDateTime& Expiration, const TSharedRef<IMessageSender, ESPMode::ThreadSafe>& Publisher) override;
+	virtual void Register(const FMessageAddress& Address, const TSharedRef<IMessageReceiver, ESPMode::ThreadSafe>& Recipient) override;
+	virtual void Send(void* Message, UScriptStruct* TypeInfo, const TSharedPtr<IMessageAttachment, ESPMode::ThreadSafe>& Attachment, const TArray<FMessageAddress>& Recipients, const FTimespan& Delay, const FDateTime& Expiration, const TSharedRef<IMessageSender, ESPMode::ThreadSafe>& Sender) override;
 	virtual void Shutdown() override;
-	virtual IMessageSubscriptionPtr Subscribe(const IReceiveMessagesRef& Subscriber, const FName& MessageType, const FMessageScopeRange& ScopeRange) override;
-	virtual void Unintercept(const IMessageInterceptorRef& Interceptor, const FName& MessageType) override;
+	virtual TSharedPtr<IMessageSubscription, ESPMode::ThreadSafe> Subscribe(const TSharedRef<IMessageReceiver, ESPMode::ThreadSafe>& Subscriber, const FName& MessageType, const FMessageScopeRange& ScopeRange) override;
+	virtual void Unintercept(const TSharedRef<IMessageInterceptor, ESPMode::ThreadSafe>& Interceptor, const FName& MessageType) override;
 	virtual void Unregister(const FMessageAddress& Address) override;
-	virtual void Unsubscribe(const IReceiveMessagesRef& Subscriber, const FName& MessageType) override;
+	virtual void Unsubscribe(const TSharedRef<IMessageReceiver, ESPMode::ThreadSafe>& Subscriber, const FName& MessageType) override;
 
 private:
 
@@ -62,7 +61,7 @@ private:
 	FRunnableThread* RouterThread;
 
 	/** Holds the recipient authorizer. */
-	IAuthorizeMessageRecipientsPtr RecipientAuthorizer;
+	TSharedPtr<IAuthorizeMessageRecipients> RecipientAuthorizer;
 
 	/** Holds bus shutdown delegate. */
 	FOnMessageBusShutdown ShutdownDelegate;
diff --git a/Engine/Source/Runtime/Messaging/Private/Bus/MessageContext.cpp b/Engine/Source/Runtime/Messaging/Private/Bus/MessageContext.cpp
--- a/Engine/Source/Runtime/Messaging/Private/Bus/MessageContext.cpp
+++ b/Engine/Source/Runtime/Messaging/Private/Bus/MessageContext.cpp
@@ -1,6 +1,7 @@
 // Copyright 1998-2016 Epic Games, Inc. All Rights Reserved.
 
 #include "MessagingPrivatePCH.h"
+#include "MessageContext.h"
 
 
 /* FMessageContext structors
@@ -34,7 +35,7 @@ const TMap<FName, FString>& FMessageContext::GetAnnotations() const
 }
 
 
-IMessageAttachmentPtr FMessageContext::GetAttachment() const
+TSharedPtr<IMessageAttachment, ESPMode::ThreadSafe> FMessageContext::GetAttachment() const
 {
 	if (OriginalContext.IsValid())
 	{
diff --git a/Engine/Source/Runtime/Messaging/Private/Bus/MessageContext.h b/Engine/Source/Runtime/Messaging/Private/Bus/MessageContext.h
--- a/Engine/Source/Runtime/Messaging/Private/Bus/MessageContext.h
+++ b/Engine/Source/Runtime/Messaging/Private/Bus/MessageContext.h
@@ -43,7 +43,17 @@ public:
 	 * @param InExpiration The message's expiration time.
 	 * @param InSenderThread The name of the thread from which the message was sent.
 	 */
-	FMessageContext(void* InMessage, UScriptStruct* InTypeInfo, const IMessageAttachmentPtr& InAttachment, const FMessageAddress& InSender, const TArray<FMessageAddress>& InRecipients, EMessageScope InScope, const FDateTime& InTimeSent, const FDateTime& InExpiration, ENamedThreads::Type InSenderThread)
+	FMessageContext(
+		void* InMessage,
+		UScriptStruct* InTypeInfo,
+		const TSharedPtr<IMessageAttachment, ESPMode::ThreadSafe>& InAttachment,
+		const FMessageAddress& InSender,
+		const TArray<FMessageAddress>& InRecipients,
+		EMessageScope InScope,
+		const FDateTime& InTimeSent,
+		const FDateTime& InExpiration,
+		ENamedThreads::Type InSenderThread
+	)
 		: Attachment(InAttachment)
 		, Expiration(InExpiration)
 		, Message(InMessage)
@@ -67,7 +77,14 @@ public:
 	 * @param InTimeForwarded The time at which the message was forwarded.
 	 * @param InForwarderThread The name of the thread from which the message was forwarded.
 	 */
-	FMessageContext(const IMessageContextRef& InContext, const FMessageAddress& InForwarder, const TArray<FMessageAddress>& NewRecipients, EMessageScope NewScope, const FDateTime& InTimeForwarded, ENamedThreads::Type InForwarderThread)
+	FMessageContext(
+		const IMessageContextRef& InContext,
+		const FMessageAddress& InForwarder,
+		const TArray<FMessageAddress>& NewRecipients,
+		EMessageScope NewScope,
+		const FDateTime& InTimeForwarded,
+		ENamedThreads::Type InForwarderThread
+	)
 		: Message(nullptr)
 		, OriginalContext(InContext)
 		, Recipients(NewRecipients)
@@ -82,10 +99,10 @@ public:
 
 public:
 
-	// IMessageContext interface
+	//~ IMessageContext interface
 
 	virtual const TMap<FName, FString>& GetAnnotations() const override;
-	virtual IMessageAttachmentPtr GetAttachment() const override;
+	virtual TSharedPtr<IMessageAttachment, ESPMode::ThreadSafe> GetAttachment() const override;
 	virtual const FDateTime& GetExpiration() const override;
 	virtual const void* GetMessage() const override;
 	virtual const TWeakObjectPtr<UScriptStruct>& GetMessageTypeInfo() const override;
@@ -103,7 +120,7 @@ private:
 	TMap<FName, FString> Annotations;
 
 	/** Holds a pointer to attached binary data. */
-	IMessageAttachmentPtr Attachment;
+	TSharedPtr<IMessageAttachment, ESPMode::ThreadSafe> Attachment;
 
 	/** Holds the expiration time. */
 	FDateTime Expiration;
diff --git a/Engine/Source/Runtime/Messaging/Private/Bus/MessageDispatchTask.h b/Engine/Source/Runtime/Messaging/Private/Bus/MessageDispatchTask.h
--- a/Engine/Source/Runtime/Messaging/Private/Bus/MessageDispatchTask.h
+++ b/Engine/Source/Runtime/Messaging/Private/Bus/MessageDispatchTask.h
@@ -2,6 +2,13 @@
 
 #pragma once
 
+#include "TaskGraphInterfaces.h"
+
+
+class FMessageTracer;
+class IMessageContext;
+class IMessageReceiver;
+
 
 /**
  * Implements an asynchronous task for dispatching a message to a recipient.
@@ -18,12 +25,11 @@ public:
 	 * @param InRecipient The message recipient.
 	 * @param InTracer The message tracer to notify.
 	 */
-	FMessageDispatchTask(ENamedThreads::Type InThread, IMessageContextRef InContext, IReceiveMessagesWeakPtr InRecipient, FMessageTracerPtr InTracer)
-		: Context(InContext)
-		, RecipientPtr(InRecipient)
-		, Thread(InThread)
-		, TracerPtr(InTracer)
-	{ }
+	FMessageDispatchTask(
+		ENamedThreads::Type InThread,
+		TSharedRef<IMessageContext, ESPMode::ThreadSafe> InContext,
+		TWeakPtr<IMessageReceiver, ESPMode::ThreadSafe> InRecipient,
+		TSharedPtr<FMessageTracer, ESPMode::ThreadSafe> InTracer);
 
 public:
 
@@ -33,29 +39,7 @@ public:
 	 * @param CurrentThread The thread that this task is executing on.
 	 * @param MyCompletionGraphEvent The completion event.
 	 */
-	void DoTask(ENamedThreads::Type CurrentThread, const FGraphEventRef& MyCompletionGraphEvent)
-	{
-		IReceiveMessagesPtr Recipient = RecipientPtr.Pin();
-
-		if (!Recipient.IsValid())
-		{
-			return;
-		}
-
-		FMessageTracerPtr Tracer = TracerPtr.Pin();
-
-		if (Tracer.IsValid())
-		{
-			Tracer->TraceDispatchedMessage(Context, Recipient.ToSharedRef(), true);
-		}
-		
-		Recipient->ReceiveMessage(Context);
-
-		if (TracerPtr.IsValid())
-		{
-			Tracer->TraceHandledMessage(Context, Recipient.ToSharedRef());
-		}
-	}
+	void DoTask(ENamedThreads::Type CurrentThread, const FGraphEventRef& MyCompletionGraphEvent);
 	
 	/**
 	 * Returns the name of the thread that this task should run on.
@@ -87,14 +71,14 @@ public:
 private:
 
 	/** Holds the message context. */
-	IMessageContextRef Context;
+	TSharedRef<IMessageContext, ESPMode::ThreadSafe> Context;
 
 	/** Holds a reference to the recipient. */
-	IReceiveMessagesWeakPtr RecipientPtr;
+	TWeakPtr<IMessageReceiver, ESPMode::ThreadSafe> RecipientPtr;
 
 	/** Holds the name of the thread that the router is running on. */
 	ENamedThreads::Type Thread;
 
 	/** Holds a pointer to the message tracer. */
-	FMessageTracerWeakPtr TracerPtr;
+	TWeakPtr<FMessageTracer, ESPMode::ThreadSafe> TracerPtr;
 };
diff --git a/Engine/Source/Runtime/Messaging/Private/Bus/MessageRouter.cpp b/Engine/Source/Runtime/Messaging/Private/Bus/MessageRouter.cpp
--- a/Engine/Source/Runtime/Messaging/Private/Bus/MessageRouter.cpp
+++ b/Engine/Source/Runtime/Messaging/Private/Bus/MessageRouter.cpp
@@ -1,6 +1,8 @@
 // Copyright 1998-2016 Epic Games, Inc. All Rights Reserved.
 
 #include "MessagingPrivatePCH.h"
+#include "MessageDispatchTask.h"
+#include "MessageRouter.h"
 
 
 /* FMessageRouter structors
@@ -69,7 +71,7 @@ void FMessageRouter::Stop()
 
 void FMessageRouter::Exit()
 {
-	TArray<IReceiveMessagesWeakPtr> Recipients;
+	TArray<TWeakPtr<IMessageReceiver, ESPMode::ThreadSafe>> Recipients;
 
 	// gather all subscribed and registered recipients
 	for (const auto& RecipientPair : ActiveRecipients)
@@ -108,11 +110,11 @@ FTimespan FMessageRouter::CalculateWaitTime()
 }
 
 
-void FMessageRouter::DispatchMessage( const IMessageContextRef& Context )
+void FMessageRouter::DispatchMessage(const IMessageContextRef& Context)
 {
 	if (Context->IsValid())
 	{
-		TArray<IReceiveMessagesPtr> Recipients;
+		TArray<TSharedPtr<IMessageReceiver, ESPMode::ThreadSafe>> Recipients;
 
 		// get recipients, either from the context...
 		const TArray<FMessageAddress>& RecipientList = Context->GetRecipients();
@@ -121,7 +123,7 @@ void FMessageRouter::DispatchMessage( const IMessageContextRef& Context )
 		{
 			for (const auto& RecipientAddress : RecipientList)
 			{
-				IReceiveMessagesPtr Recipient = ActiveRecipients.FindRef(RecipientAddress).Pin();
+				auto Recipient = ActiveRecipients.FindRef(RecipientAddress).Pin();
 
 				if (Recipient.IsValid())
 				{
@@ -160,20 +162,24 @@ void FMessageRouter::DispatchMessage( const IMessageContextRef& Context )
 }
 
 
-void FMessageRouter::FilterSubscriptions( TArray<IMessageSubscriptionPtr>& Subscriptions, const IMessageContextRef& Context, TArray<IReceiveMessagesPtr>& OutRecipients )
+void FMessageRouter::FilterSubscriptions(
+	TArray<TSharedPtr<IMessageSubscription, ESPMode::ThreadSafe>>& Subscriptions,
+	const IMessageContextRef& Context,
+	TArray<TSharedPtr<IMessageReceiver, ESPMode::ThreadSafe>>& OutRecipients
+)
 {
 	EMessageScope MessageScope = Context->GetScope();
 
 	for (int32 SubscriptionIndex = 0; SubscriptionIndex < Subscriptions.Num(); ++SubscriptionIndex)
 	{
-		const IMessageSubscriptionPtr& Subscription = Subscriptions[SubscriptionIndex];
+		const auto& Subscription = Subscriptions[SubscriptionIndex];
 
 		if (!Subscription->IsEnabled() || !Subscription->GetScopeRange().Contains(MessageScope))
 		{
 			continue;
 		}
 
-		IReceiveMessagesPtr Subscriber = Subscription->GetSubscriber().Pin();
+		auto Subscriber = Subscription->GetSubscriber().Pin();
 
 		if (Subscriber.IsValid())
 		{
@@ -214,16 +220,16 @@ void FMessageRouter::ProcessDelayedMessages()
 /* FMessageRouter callbacks
  *****************************************************************************/
 
-void FMessageRouter::HandleAddInterceptor( IMessageInterceptorRef Interceptor, FName MessageType )
+void FMessageRouter::HandleAddInterceptor(TSharedRef<IMessageInterceptor, ESPMode::ThreadSafe> Interceptor, FName MessageType)
 {
 	ActiveInterceptors.FindOrAdd(MessageType).AddUnique(Interceptor);
 	Tracer->TraceAddedInterceptor(Interceptor, MessageType);
 }
 
 
-void FMessageRouter::HandleAddRecipient( FMessageAddress Address, IReceiveMessagesWeakPtr RecipientPtr )
+void FMessageRouter::HandleAddRecipient(FMessageAddress Address, TWeakPtr<IMessageReceiver, ESPMode::ThreadSafe> RecipientPtr)
 {
-	IReceiveMessagesPtr Recipient = RecipientPtr.Pin();
+	auto Recipient = RecipientPtr.Pin();
 
 	if (Recipient.IsValid())
 	{
@@ -233,14 +239,14 @@ void FMessageRouter::HandleAddRecipient( FMessageAddress Address, IReceiveMessag
 }
 
 
-void FMessageRouter::HandleAddSubscriber( IMessageSubscriptionRef Subscription )
+void FMessageRouter::HandleAddSubscriber(TSharedRef<IMessageSubscription, ESPMode::ThreadSafe> Subscription)
 {
 	ActiveSubscriptions.FindOrAdd(Subscription->GetMessageType()).AddUnique(Subscription);
 	Tracer->TraceAddedSubscription(Subscription);
 }
 
 
-void FMessageRouter::HandleRemoveInterceptor( IMessageInterceptorRef Interceptor, FName MessageType )
+void FMessageRouter::HandleRemoveInterceptor(TSharedRef<IMessageInterceptor, ESPMode::ThreadSafe> Interceptor, FName MessageType)
 {
 	if (MessageType == NAME_All)
 	{
@@ -251,7 +257,7 @@ void FMessageRouter::HandleRemoveInterceptor( IMessageInterceptorRef Interceptor
 	}
 	else
 	{
-		TArray<IMessageInterceptorPtr>& Interceptors = ActiveInterceptors.FindOrAdd(MessageType);
+		auto& Interceptors = ActiveInterceptors.FindOrAdd(MessageType);
 		Interceptors.Remove(Interceptor);
 	}
 
@@ -259,9 +265,9 @@ void FMessageRouter::HandleRemoveInterceptor( IMessageInterceptorRef Interceptor
 }
 
 
-void FMessageRouter::HandleRemoveRecipient( FMessageAddress Address )
+void FMessageRouter::HandleRemoveRecipient(FMessageAddress Address)
 {
-	IReceiveMessagesPtr Recipient = ActiveRecipients.FindRef(Address).Pin();
+	auto Recipient = ActiveRecipients.FindRef(Address).Pin();
 
 	if (Recipient.IsValid())
 	{
@@ -272,9 +278,9 @@ void FMessageRouter::HandleRemoveRecipient( FMessageAddress Address )
 }
 
 
-void FMessageRouter::HandleRemoveSubscriber( IReceiveMessagesWeakPtr SubscriberPtr, FName MessageType )
+void FMessageRouter::HandleRemoveSubscriber(TWeakPtr<IMessageReceiver, ESPMode::ThreadSafe> SubscriberPtr, FName MessageType)
 {
-	IReceiveMessagesPtr Subscriber = SubscriberPtr.Pin();
+	auto Subscriber = SubscriberPtr.Pin();
 
 	if (!Subscriber.IsValid())
 	{
@@ -288,7 +294,7 @@ void FMessageRouter::HandleRemoveSubscriber( IReceiveMessagesWeakPtr SubscriberP
 			continue;
 		}
 
-		TArray<IMessageSubscriptionPtr>& Subscriptions = SubscriptionsPair.Value;
+		TArray<TSharedPtr<IMessageSubscription, ESPMode::ThreadSafe>>& Subscriptions = SubscriptionsPair.Value;
 
 		for (int32 Index = 0; Index < Subscriptions.Num(); Index++)
 		{
@@ -306,12 +312,12 @@ void FMessageRouter::HandleRemoveSubscriber( IReceiveMessagesWeakPtr SubscriberP
 }
 
 
-void FMessageRouter::HandleRouteMessage( IMessageContextRef Context )
+void FMessageRouter::HandleRouteMessage(IMessageContextRef Context)
 {
 	Tracer->TraceRoutedMessage(Context);
 
 	// intercept routing
-	TArray<IMessageInterceptorPtr>& Interceptors = ActiveInterceptors.FindOrAdd(Context->GetMessageType());
+	auto& Interceptors = ActiveInterceptors.FindOrAdd(Context->GetMessageType());
 
 	for (auto& Interceptor : Interceptors)
 	{
