Copying //UE4/Dev-Sequencer to //UE4/Dev-Main (Source: //UE4/Dev-Sequencer @ 3237992)

diff --git a/Engine/Source/Editor/UnrealEd/Private/LevelEditorViewport.cpp b/Engine/Source/Editor/UnrealEd/Private/LevelEditorViewport.cpp
--- a/Engine/Source/Editor/UnrealEd/Private/LevelEditorViewport.cpp
+++ b/Engine/Source/Editor/UnrealEd/Private/LevelEditorViewport.cpp
@@ -2573,6 +2573,16 @@ void FLevelEditorViewportClient::TrackingStarted( const FInputEventState& InInpu
 		if (bIsDraggingWidget)
 		{
 			Widget->SetSnapEnabled(true);
+
+			for (FSelectedEditableComponentIterator It(GEditor->GetSelectedEditableComponentIterator()); It; ++It)
+			{
+				USceneComponent* SceneComponent = Cast<USceneComponent>(*It);
+				if (SceneComponent)
+				{
+					// Notify that this component is beginning to move
+					GEditor->BroadcastBeginObjectMovement(*SceneComponent);
+				}
+			}
 		}
 	}
 	else
@@ -2743,6 +2753,7 @@ void FLevelEditorViewportClient::TrackingStopped()
 				}
 			}
 
+			bool bComponentsMoved = false;
 			if (GEditor->GetSelectedComponentCount() > 0)
 			{
 				USelection* ComponentSelection = GEditor->GetSelectedComponents();
@@ -2783,11 +2794,19 @@ void FLevelEditorViewportClient::TrackingStopped()
 				for (USceneComponent* SceneComp : ComponentsToMove)
 				{
 					SceneComp->PostEditComponentMove(true);
+
+					GEditor->BroadcastEndObjectMovement(*SceneComp);
+
+					bComponentsMoved = true;
 				}
 			}
-
-			Actor->PostEditMove(true);
-			GEditor->BroadcastEndObjectMovement(*Actor);
+				
+			if (!bComponentsMoved)
+			{
+				Actor->PostEditMove(true);
+	
+				GEditor->BroadcastEndObjectMovement(*Actor);
+			}
 		}
 
 		if (!GUnrealEd->IsPivotMovedIndependently())
diff --git a/Engine/Source/Editor/UnrealEd/Private/SCurveEditor.cpp b/Engine/Source/Editor/UnrealEd/Private/SCurveEditor.cpp
--- a/Engine/Source/Editor/UnrealEd/Private/SCurveEditor.cpp
+++ b/Engine/Source/Editor/UnrealEd/Private/SCurveEditor.cpp
@@ -69,6 +69,7 @@ void SCurveEditor::Construct(const FArguments& InArgs)
 	OutputSnap = InArgs._OutputSnap;
 	bInputSnappingEnabled = InArgs._InputSnappingEnabled;
 	bOutputSnappingEnabled = InArgs._OutputSnappingEnabled;
+	bShowTimeInFrames = InArgs._ShowTimeInFrames;
 
 	bZoomToFitVertical = InArgs._ZoomToFitVertical;
 	bZoomToFitHorizontal = InArgs._ZoomToFitHorizontal;
@@ -290,6 +291,7 @@ void SCurveEditor::Construct(const FArguments& InArgs)
 	CurveSelectionWidget = CurveSelector;
 
 	InputAxisName = InArgs._XAxisName.IsSet() ? FText::FromString(InArgs._XAxisName.GetValue()) : LOCTEXT("Time", "Time");
+	InputFrameAxisName = InArgs._XAxisName.IsSet() ? FText::FromString(InArgs._XAxisName.GetValue()) : LOCTEXT("Frame", "Frame");
 	OutputAxisName = InArgs._YAxisName.IsSet() ? FText::FromString(InArgs._YAxisName.GetValue()) : LOCTEXT("Value", "Value");
 
 	this->ChildSlot
@@ -364,27 +366,62 @@ void SCurveEditor::Construct(const FArguments& InArgs)
 						.Visibility(this, &SCurveEditor::GetEditVisibility)
 						.VAlign(VAlign_Center)
 						[
-							SNew(SNumericEntryBox<float>)
-							.IsEnabled(this, &SCurveEditor::GetInputEditEnabled)
-							.Font(FEditorStyle::GetFontStyle("CurveEd.InfoFont"))
-							.Value(this, &SCurveEditor::OnGetTime)
-							.UndeterminedString(LOCTEXT("MultipleValues", "Multiple Values"))
-							.OnValueCommitted(this, &SCurveEditor::OnTimeComitted)
-							.OnValueChanged(this, &SCurveEditor::OnTimeChanged)
-							.OnBeginSliderMovement(this, &SCurveEditor::OnBeginSliderMovement, LOCTEXT("SetTime", "Set New Time"))
-							.OnEndSliderMovement(this, &SCurveEditor::OnEndSliderMovement)
-							.LabelVAlign(VAlign_Center)
-							.AllowSpin(true)
-							.MinValue(TOptional<float>())
-							.MaxValue(TOptional<float>())
-							.MaxSliderValue(TOptional<float>())
-							.MinSliderValue(TOptional<float>())
-							.Delta(this, &SCurveEditor::GetInputNumericEntryBoxDelta)
-							.MinDesiredValueWidth(60.0f)
-							.Label()
+							SNew(SHorizontalBox)
+
+							+ SHorizontalBox::Slot()
+							.AutoWidth()
 							[
-								SNew(STextBlock)
-								.Text(InputAxisName)
+								SNew(SNumericEntryBox<float>)
+								.IsEnabled(this, &SCurveEditor::GetInputEditEnabled)
+								.Font(FEditorStyle::GetFontStyle("CurveEd.InfoFont"))
+								.Value(this, &SCurveEditor::OnGetTime)
+								.UndeterminedString(LOCTEXT("MultipleValues", "Multiple Values"))
+								.OnValueCommitted(this, &SCurveEditor::OnTimeComitted)
+								.OnValueChanged(this, &SCurveEditor::OnTimeChanged)
+								.OnBeginSliderMovement(this, &SCurveEditor::OnBeginSliderMovement, LOCTEXT("SetTime", "Set New Time"))
+								.OnEndSliderMovement(this, &SCurveEditor::OnEndSliderMovement)
+								.LabelVAlign(VAlign_Center)
+								.AllowSpin(true)
+								.MinValue(TOptional<float>())
+								.MaxValue(TOptional<float>())
+								.MaxSliderValue(TOptional<float>())
+								.MinSliderValue(TOptional<float>())
+								.Delta(this, &SCurveEditor::GetInputNumericEntryBoxDelta)
+								.MinDesiredValueWidth(60.0f)
+								.Visibility(this, &SCurveEditor::GetTimeEditVisibility)
+								.Label()
+								[
+									SNew(STextBlock)
+									.Text(this, &SCurveEditor::GetInputAxisName)
+								]
+							]
+
+							+ SHorizontalBox::Slot()
+							.AutoWidth()
+							[
+								SNew(SNumericEntryBox<int32>)
+								.IsEnabled(this, &SCurveEditor::GetInputEditEnabled)
+								.Font(FEditorStyle::GetFontStyle("CurveEd.InfoFont"))
+								.Value(this, &SCurveEditor::OnGetTimeInFrames)
+								.UndeterminedString(LOCTEXT("MultipleValues", "Multiple Values"))
+								.OnValueCommitted(this, &SCurveEditor::OnTimeInFramesComitted)
+								.OnValueChanged(this, &SCurveEditor::OnTimeInFramesChanged)
+								.OnBeginSliderMovement(this, &SCurveEditor::OnBeginSliderMovement, LOCTEXT("SetFrame", "Set New Frame"))
+								.OnEndSliderMovement(this, &SCurveEditor::OnEndSliderMovement)
+								.LabelVAlign(VAlign_Center)
+								.AllowSpin(true)
+								.MinValue(TOptional<int32>())
+								.MaxValue(TOptional<int32>())
+								.MaxSliderValue(TOptional<int32>())
+								.MinSliderValue(TOptional<int32>())
+								.Delta(1)
+								.MinDesiredValueWidth(60.0f)
+								.Visibility(this, &SCurveEditor::GetFrameEditVisibility)
+								.Label()
+								[
+									SNew(STextBlock)
+									.Text(this, &SCurveEditor::GetInputAxisName)
+								]
 							]
 						]
 					]
@@ -533,6 +570,11 @@ FText SCurveEditor::GetCurveToolTipOutputText() const
 	return CurveToolTipOutputText;
 }
 
+FText SCurveEditor::GetInputAxisName() const
+{
+	return ShowTimeInFrames() ? InputFrameAxisName : InputAxisName;
+}
+
 SCurveEditor::~SCurveEditor()
 {
 	if (GEditor != NULL)
@@ -687,6 +729,16 @@ EVisibility SCurveEditor::GetZoomButtonVisibility() const
 	return (IsHovered() || (false == bHideUI)) && bShowZoomButtons ? EVisibility::Visible : EVisibility::Collapsed;
 }
 
+EVisibility SCurveEditor::GetTimeEditVisibility() const
+{
+	return ShowTimeInFrames() ? EVisibility::Collapsed : EVisibility::Visible;
+}
+
+EVisibility SCurveEditor::GetFrameEditVisibility() const
+{
+	return ShowTimeInFrames() ? EVisibility::Visible : EVisibility::Collapsed;
+}
+
 bool SCurveEditor::GetInputEditEnabled() const
 {
 	return (SelectedKeys.Num() == 1);
@@ -1455,9 +1507,17 @@ void SCurveEditor::UpdateCurveToolTip(const FGeometry& InMyGeometry, const FPoin
 			FNumberFormattingOptions FormattingOptions;
 			FormattingOptions.MaximumFractionalDigits = 2;
 			CurveToolTipNameText = FText::FromName(HoveredCurve->CurveInfo.CurveName);
-			CurveToolTipInputText = FText::Format(LOCTEXT("CurveToolTipTimeFormat", "{0}: {1}"), InputAxisName, FText::AsNumber(Time, &FormattingOptions));
 			CurveToolTipOutputText = FText::Format(LOCTEXT("CurveToolTipValueFormat", "{0}: {1}"), OutputAxisName, FText::AsNumber(Value, &FormattingOptions));
 
+			if (ShowTimeInFrames())
+			{
+				CurveToolTipInputText = FText::Format(LOCTEXT("CurveToolTipFrameFormat", "{0}: {1}"), GetInputAxisName(), FText::AsNumber(TimeToFrame(Time)));
+			}
+			else
+			{
+				CurveToolTipInputText = FText::Format(LOCTEXT("CurveToolTipTimeFormat", "{0}: {1}"), GetInputAxisName(), FText::AsNumber(Time, &FormattingOptions));
+			}
+
 			if (CurveToolTip.IsValid() == false)
 			{
 				SetToolTip(
@@ -2028,6 +2088,64 @@ void SCurveEditor::OnTimeChanged(float NewTime)
 	}
 }
 
+TOptional<int32> SCurveEditor::OnGetTimeInFrames() const
+{
+	if ( SelectedKeys.Num() == 1 )
+	{
+		TOptional<float> KeyTime = GetKeyTime(SelectedKeys[0]);
+		if (KeyTime.IsSet())
+		{
+			return TOptional<int32>(TimeToFrame(KeyTime.GetValue()));
+		}
+	}
+
+	// Value couldn't be accessed.  Return an unset value
+	return TOptional<int32>();
+}
+
+void SCurveEditor::OnTimeInFramesComitted(int32 NewFrame, ETextCommit::Type CommitType)
+{
+	// Don't digest the number if we just clicked away from the pop-up
+	if ( !bIsUsingSlider && ((CommitType == ETextCommit::OnEnter) || ( CommitType == ETextCommit::OnUserMovedFocus )) )
+	{
+		if ( SelectedKeys.Num() >= 1 )
+		{
+			auto Key = SelectedKeys[0];
+			if ( IsValidCurve(Key.Curve) )
+			{
+				const FScopedTransaction Transaction(LOCTEXT("CurveEditor_NewFrame", "New Frame Entered"));
+				CurveOwner->ModifyOwner();
+				Key.Curve->SetKeyTime(Key.KeyHandle, FrameToTime(NewFrame));
+				TArray<FRichCurveEditInfo> ChangedCurveEditInfos;
+				ChangedCurveEditInfos.Add(GetViewModelForCurve(Key.Curve)->CurveInfo);
+				CurveOwner->OnCurveChanged(ChangedCurveEditInfos);
+			}
+		}
+
+		FSlateApplication::Get().DismissAllMenus();
+	}
+}
+
+void SCurveEditor::OnTimeInFramesChanged(int32 NewFrame)
+{
+	if ( bIsUsingSlider )
+	{
+		if ( SelectedKeys.Num() >= 1 )
+		{
+			auto Key = SelectedKeys[0];
+			if ( IsValidCurve(Key.Curve) )
+			{
+				const FScopedTransaction Transaction( LOCTEXT( "CurveEditor_NewFrame", "New Frame Entered" ) );
+				CurveOwner->ModifyOwner();
+				Key.Curve->SetKeyTime(Key.KeyHandle, FrameToTime(NewFrame));
+				TArray<FRichCurveEditInfo> ChangedCurveEditInfos;
+				ChangedCurveEditInfos.Add(GetViewModelForCurve(Key.Curve)->CurveInfo);
+				CurveOwner->OnCurveChanged(ChangedCurveEditInfos);
+			}
+		}
+	}
+}
+
 TOptional<float> SCurveEditor::OnGetValue() const
 {
 	TOptional<float> Value;
@@ -2132,6 +2250,13 @@ void SCurveEditor::OnEndSliderMovement(float NewValue)
 	GEditor->EndTransaction();
 }
 
+void SCurveEditor::OnEndSliderMovement(int32 NewValue)
+{
+	bIsUsingSlider = false;
+
+	GEditor->EndTransaction();
+}
+
 
 SCurveEditor::FSelectedCurveKey SCurveEditor::HitTestKeys(const FGeometry& InMyGeometry, const FVector2D& HitScreenPosition)
 {	
@@ -2570,16 +2695,21 @@ void SCurveEditor::ToggleOutputSnapping()
 	}
 }
 
-bool SCurveEditor::IsInputSnappingEnabled()
+bool SCurveEditor::IsInputSnappingEnabled() const
 {
 	return bInputSnappingEnabled.Get();
 }
 
-bool SCurveEditor::IsOutputSnappingEnabled()
+bool SCurveEditor::IsOutputSnappingEnabled() const
 {
 	return bOutputSnappingEnabled.Get();
 }
 
+bool SCurveEditor::ShowTimeInFrames() const
+{
+	return bShowTimeInFrames.Get();
+}
+
 void SCurveEditor::CreateContextMenu(const FGeometry& InMyGeometry, const FPointerEvent& InMouseEvent)
 {
 	TSharedPtr<TArray<TSharedPtr<FCurveViewModel>>> CurvesToAddKeysTo = MakeShareable(new TArray<TSharedPtr<FCurveViewModel>>());
@@ -3659,4 +3789,19 @@ void SCurveEditor::CloseEntryPopupMenu()
 	}
 }
 
+	
+int32 SCurveEditor::TimeToFrame(float InTime) const
+{
+	const float FrameRate = InputSnap.IsSet() ? 1.0f / InputSnap.Get() : 1.f;
+	float Frame = InTime * FrameRate;
+	return FMath::RoundToInt(Frame);
+}
+
+
+float SCurveEditor::FrameToTime(int32 InFrame) const
+{
+	const float FrameRate = InputSnap.IsSet() ? 1.0f / InputSnap.Get() : 1.f;
+	return InFrame / FrameRate;
+}
+
 #undef LOCTEXT_NAMESPACE
diff --git a/Engine/Source/Editor/UnrealEd/Public/CurveEditorSettings.h b/Engine/Source/Editor/UnrealEd/Public/CurveEditorSettings.h
--- a/Engine/Source/Editor/UnrealEd/Public/CurveEditorSettings.h
+++ b/Engine/Source/Editor/UnrealEd/Public/CurveEditorSettings.h
@@ -70,16 +70,16 @@ public:
 	FOnCurveEditorCurveVisibilityChanged& GetOnCurveEditorCurveVisibilityChanged();
 
 protected:
-	UPROPERTY( config )
+	UPROPERTY( config, EditAnywhere, Category="Curve Editor" )
 	bool bAutoFrameCurveEditor;
 
-	UPROPERTY( config )
+	UPROPERTY( config, EditAnywhere, Category="Curve Editor" )
 	bool bShowCurveEditorCurveToolTips;
 
-	UPROPERTY( config )
+	UPROPERTY( config, EditAnywhere, Category="Curve Editor" )
 	TEnumAsByte<ECurveEditorCurveVisibility::Type> CurveVisibility;
 
-	UPROPERTY( config )
+	UPROPERTY( config, EditAnywhere, Category="Curve Editor" )
 	TEnumAsByte<ECurveEditorTangentVisibility::Type> TangentVisibility;
 
 	FOnCurveEditorCurveVisibilityChanged OnCurveEditorCurveVisibilityChanged;
diff --git a/Engine/Source/Editor/UnrealEd/Public/SCurveEditor.h b/Engine/Source/Editor/UnrealEd/Public/SCurveEditor.h
--- a/Engine/Source/Editor/UnrealEd/Public/SCurveEditor.h
+++ b/Engine/Source/Editor/UnrealEd/Public/SCurveEditor.h
@@ -177,6 +177,7 @@ public:
 		, _OutputSnap(0.05f)
 		, _InputSnappingEnabled(false)
 		, _OutputSnappingEnabled(false)
+		, _ShowTimeInFrames(false)
 		, _TimelineLength(5.0f)
 		, _DesiredSize(FVector2D::ZeroVector)
 		, _DrawCurve(true)
@@ -204,6 +205,7 @@ public:
 		SLATE_ATTRIBUTE( float, OutputSnap )
 		SLATE_ATTRIBUTE( bool, InputSnappingEnabled )
 		SLATE_ATTRIBUTE( bool, OutputSnappingEnabled )
+		SLATE_ATTRIBUTE( bool, ShowTimeInFrames )
 		SLATE_ATTRIBUTE( float, TimelineLength )
 		SLATE_ATTRIBUTE( FVector2D, DesiredSize )
 		SLATE_ARGUMENT( bool, DrawCurve )
@@ -411,25 +413,33 @@ private:
 
 	void ToggleInputSnapping();
 	void ToggleOutputSnapping();
-	bool IsInputSnappingEnabled();
-	bool IsOutputSnappingEnabled();
+	bool IsInputSnappingEnabled() const;
+	bool IsOutputSnappingEnabled() const;
+	bool ShowTimeInFrames() const;
 
 	TOptional<float> OnGetTime() const;
 	void OnTimeComitted(float NewValue, ETextCommit::Type CommitType);
 	void OnTimeChanged(float NewValue);
 
+	TOptional<int32> OnGetTimeInFrames() const;
+	void OnTimeInFramesComitted(int32 NewValue, ETextCommit::Type CommitType);
+	void OnTimeInFramesChanged(int32 NewValue);
+
 	TOptional<float> OnGetValue() const;
 	void OnValueComitted(float NewValue, ETextCommit::Type CommitType);
 	void OnValueChanged(float NewValue);
 
 	void OnBeginSliderMovement(FText TransactionName);
 	void OnEndSliderMovement(float NewValue);
+	void OnEndSliderMovement(int32 NewValue);
 
 	EVisibility GetCurveAreaVisibility() const;
 	EVisibility GetCurveSelectorVisibility() const;
 	EVisibility GetEditVisibility() const;
 	EVisibility GetColorGradientVisibility() const;
 	EVisibility GetZoomButtonVisibility() const;
+	EVisibility GetTimeEditVisibility() const;
+	EVisibility GetFrameEditVisibility() const;
 
 	bool GetInputEditEnabled() const;
 
@@ -585,6 +595,8 @@ private:
 	FText GetCurveToolTipInputText() const;
 	FText GetCurveToolTipOutputText() const;
 
+	FText GetInputAxisName() const;
+
 	void UpdateCurveToolTip( const FGeometry& MyGeometry, const FPointerEvent& MouseEvent );
 
 	TSharedPtr<FCurveViewModel> GetViewModelForCurve(FRichCurve* InCurve);
@@ -627,6 +639,10 @@ protected:
 	/** Closes the popup created by GenericTextEntryModeless*/
 	void CloseEntryPopupMenu();
 
+	/** Convert time to frames and vice versa */
+	int32 TimeToFrame(float InTime) const;
+	float FrameToTime(int32 InFrame) const;
+
 private:
 
 	/** User-supplied object for this curve editor */
@@ -725,6 +741,9 @@ protected:
 	/** Whether or not output snapping is enabled. */
 	TAttribute<bool> bOutputSnappingEnabled;
 
+	/** Show time in frames. */
+	TAttribute<bool> bShowTimeInFrames;
+
 	/** True if you want the curve editor to fit to zoom **/
 	bool bZoomToFitVertical;
 
@@ -766,6 +785,8 @@ protected:
 
 	/** The text to display for the input axis. */
 	FText InputAxisName;
+	/** The text to display for the input (frame) axis. */
+	FText InputFrameAxisName;
 	/** The text to display for the output axis. */
 	FText OutputAxisName;
 
diff --git a/Engine/Source/Editor/UnrealEd/Public/Toolkits/AssetEditorToolkit.h b/Engine/Source/Editor/UnrealEd/Public/Toolkits/AssetEditorToolkit.h
--- a/Engine/Source/Editor/UnrealEd/Public/Toolkits/AssetEditorToolkit.h
+++ b/Engine/Source/Editor/UnrealEd/Public/Toolkits/AssetEditorToolkit.h
@@ -312,6 +312,9 @@ protected:
 	/** The asset editor mode manager */
 	FAssetEditorModeManager* AssetEditorModeManager;
 
+	/** Array of layout extenders */
+	TArray<TSharedPtr<FLayoutExtender>> LayoutExtenders;
+
 private:
 	/** The toolkit standalone host; may be nullptr */
 	TWeakPtr<SStandaloneAssetEditorToolkitHost> StandaloneHost;
diff --git a/Engine/Source/Programs/UnrealBuildTool/Android/UEDeployAndroid.cs b/Engine/Source/Programs/UnrealBuildTool/Android/UEDeployAndroid.cs
--- a/Engine/Source/Programs/UnrealBuildTool/Android/UEDeployAndroid.cs
+++ b/Engine/Source/Programs/UnrealBuildTool/Android/UEDeployAndroid.cs
@@ -1422,9 +1422,10 @@ namespace UnrealBuildTool
 						break;
 				}
 			}
+			bool bSupportingAllTextureFormats = bETC1Enabled && bETC2Enabled && bDXTEnabled && bATCEnabled && bPVRTCEnabled && bASTCEnabled;
 
 			// If it is only ETC2 we need to skip adding the texture format filtering and instead use ES 3.0 as minimum version (it requires ETC2)
-			bool bOnlyETC2Enabled = (bETC2Enabled && !(bETC1Enabled || bDXTEnabled || bATCEnabled || bASTCEnabled));
+			bool bOnlyETC2Enabled = (bETC2Enabled && !(bETC1Enabled || bDXTEnabled || bATCEnabled || bPVRTCEnabled || bASTCEnabled));
 
 			StringBuilder Text = new StringBuilder();
 			Text.AppendLine("<?xml version=\"1.0\" encoding=\"utf-8\"?>");
@@ -1620,34 +1621,37 @@ namespace UnrealBuildTool
 					Text.AppendLine("\t<uses-permission android:name=\"com.qti.permission.PROFILER\"/>");
 				}
 
-				Text.AppendLine("\t<!-- Supported texture compression formats (cooked) -->");
-				if (bETC1Enabled)
+				if (!bSupportingAllTextureFormats)
 				{
-					Text.AppendLine("\t<supports-gl-texture android:name=\"GL_OES_compressed_ETC1_RGB8_texture\" />");
-				}
-				if (bETC2Enabled && !bOnlyETC2Enabled)
-				{
-					Text.AppendLine("\t<supports-gl-texture android:name=\"GL_COMPRESSED_RGB8_ETC2\" />");
-					Text.AppendLine("\t<supports-gl-texture android:name=\"GL_COMPRESSED_RGBA8_ETC2_EAC\" />");
-				}
-				if (bATCEnabled)
-				{
-					Text.AppendLine("\t<supports-gl-texture android:name=\"GL_AMD_compressed_ATC_texture\" />");
-					Text.AppendLine("\t<supports-gl-texture android:name=\"GL_ATI_texture_compression_atitc\" />");
-				}
-				if (bDXTEnabled)
-				{
-					Text.AppendLine("\t<supports-gl-texture android:name=\"GL_EXT_texture_compression_dxt1\" />");
-					Text.AppendLine("\t<supports-gl-texture android:name=\"GL_EXT_texture_compression_s3tc\" />");
-					Text.AppendLine("\t<supports-gl-texture android:name=\"GL_NV_texture_compression_s3tc\" />");
-				}
-				if (bPVRTCEnabled)
-				{
-					Text.AppendLine("\t<supports-gl-texture android:name=\"GL_IMG_texture_compression_pvrtc\" />");
-				}
-				if (bASTCEnabled)
-				{
-					Text.AppendLine("\t<supports-gl-texture android:name=\"GL_KHR_texture_compression_astc_ldr\" />");
+					Text.AppendLine("\t<!-- Supported texture compression formats (cooked) -->");
+					if (bETC1Enabled)
+					{
+						Text.AppendLine("\t<supports-gl-texture android:name=\"GL_OES_compressed_ETC1_RGB8_texture\" />");
+					}
+					if (bETC2Enabled && !bOnlyETC2Enabled)
+					{
+						Text.AppendLine("\t<supports-gl-texture android:name=\"GL_COMPRESSED_RGB8_ETC2\" />");
+						Text.AppendLine("\t<supports-gl-texture android:name=\"GL_COMPRESSED_RGBA8_ETC2_EAC\" />");
+					}
+					if (bATCEnabled)
+					{
+						Text.AppendLine("\t<supports-gl-texture android:name=\"GL_AMD_compressed_ATC_texture\" />");
+						Text.AppendLine("\t<supports-gl-texture android:name=\"GL_ATI_texture_compression_atitc\" />");
+					}
+					if (bDXTEnabled)
+					{
+						Text.AppendLine("\t<supports-gl-texture android:name=\"GL_EXT_texture_compression_dxt1\" />");
+						Text.AppendLine("\t<supports-gl-texture android:name=\"GL_EXT_texture_compression_s3tc\" />");
+						Text.AppendLine("\t<supports-gl-texture android:name=\"GL_NV_texture_compression_s3tc\" />");
+					}
+					if (bPVRTCEnabled)
+					{
+						Text.AppendLine("\t<supports-gl-texture android:name=\"GL_IMG_texture_compression_pvrtc\" />");
+					}
+					if (bASTCEnabled)
+					{
+						Text.AppendLine("\t<supports-gl-texture android:name=\"GL_KHR_texture_compression_astc_ldr\" />");
+					}
 				}
 			}
 
diff --git a/Engine/Source/Runtime/ActorSequence/ActorSequence.Build.cs b/Engine/Source/Runtime/ActorSequence/ActorSequence.Build.cs
new file mode 100644
--- /dev/null
+++ b/Engine/Source/Runtime/ActorSequence/ActorSequence.Build.cs
@@ -0,0 +1,20 @@
+// Copyright 1998-2016 Epic Games, Inc. All Rights Reserved.
+
+using UnrealBuildTool;
+
+public class ActorSequence : ModuleRules
+{
+	public ActorSequence(TargetInfo Target)
+	{
+		PrivateIncludePaths.Add("Runtime/ActorSequence/Private");
+
+		PublicDependencyModuleNames.AddRange(
+			new string[] {
+				"Core",
+				"CoreUObject",
+				"Engine",
+				"MovieScene",
+			}
+		);
+	}
+}
diff --git a/Engine/Source/Runtime/ActorSequence/Private/ActorSequencePrivatePCH.h b/Engine/Source/Runtime/ActorSequence/Private/ActorSequencePrivatePCH.h
new file mode 100644
--- /dev/null
+++ b/Engine/Source/Runtime/ActorSequence/Private/ActorSequencePrivatePCH.h
@@ -0,0 +1,11 @@
+// Copyright 1998-2016 Epic Games, Inc. All Rights Reserved.
+
+#pragma once
+
+
+/* Private dependencies
+ *****************************************************************************/
+
+#include "Core.h"
+#include "CoreUObject.h"
+#include "Engine.h"
diff --git a/Engine/Source/Runtime/ActorSequence/Public/ActorSequence.h b/Engine/Source/Runtime/ActorSequence/Public/ActorSequence.h
new file mode 100644
--- /dev/null
+++ b/Engine/Source/Runtime/ActorSequence/Public/ActorSequence.h
@@ -0,0 +1,69 @@
+// Copyright 1998-2016 Epic Games, Inc. All Rights Reserved.
+
+#pragma once
+
+#include "MovieSceneSequence.h"
+#include "MovieScene.h"
+#include "ActorSequenceObjectReference.h"
+#include "ActorSequence.generated.h"
+
+/**
+ * Movie scene animation embedded within an actor.
+ */
+UCLASS(BlueprintType, Experimental, DefaultToInstanced)
+class ACTORSEQUENCE_API UActorSequence
+	: public UMovieSceneSequence
+{
+public:
+	GENERATED_BODY()
+
+	UActorSequence(const FObjectInitializer& ObjectInitializer);
+
+	//~ UMovieSceneSequence interface
+	virtual void BindPossessableObject(const FGuid& ObjectId, UObject& PossessedObject, UObject* Context) override;
+	virtual bool CanPossessObject(UObject& Object, UObject* InPlaybackContext) const override;
+	virtual void LocateBoundObjects(const FGuid& ObjectId, UObject* Context, TArray<UObject*, TInlineAllocator<1>>& OutObjects) const override;
+	virtual UMovieScene* GetMovieScene() const override;
+	virtual UObject* GetParentObject(UObject* Object) const override;
+	virtual void UnbindPossessableObjects(const FGuid& ObjectId) override;
+
+#if WITH_EDITOR
+	virtual FText GetDisplayName() const override;
+#endif
+
+	UBlueprint* GetParentBlueprint() const;
+
+	bool IsEditable() const;
+	
+private:
+
+	//~ UObject interface
+	virtual void PostInitProperties() override;
+
+private:
+	
+	/** Pointer to the movie scene that controls this animation. */
+	UPROPERTY(Instanced)
+	UMovieScene* MovieScene;
+
+	/** Collection of object references. */
+	UPROPERTY()
+	FActorSequenceObjectReferenceMap ObjectReferences;
+
+#if WITH_EDITOR
+public:
+
+	/** Event that is fired to initialize default state for a sequence */
+	DECLARE_EVENT_OneParam(UActorSequence, FOnInitialize, UActorSequence*)
+	static FOnInitialize& OnInitializeSequence() { return OnInitializeSequenceEvent; }
+
+private:
+	static FOnInitialize OnInitializeSequenceEvent;
+#endif
+
+#if WITH_EDITORONLY_DATA
+private:
+	UPROPERTY()
+	bool bHasBeenInitialized;
+#endif
+};
\ No newline at end of file
diff --git a/Engine/Source/Runtime/ActorSequence/Public/ActorSequenceComponent.h b/Engine/Source/Runtime/ActorSequence/Public/ActorSequenceComponent.h
new file mode 100644
--- /dev/null
+++ b/Engine/Source/Runtime/ActorSequence/Public/ActorSequenceComponent.h
@@ -0,0 +1,49 @@
+// Copyright 1998-2016 Epic Games, Inc. All Rights Reserved.
+
+#pragma once
+
+#include "Components/ActorComponent.h"
+#include "MovieSceneSequencePlayer.h"
+#include "ActorSequenceComponent.generated.h"
+
+
+class UActorSequence;
+class UActorSequencePlayer;
+
+
+/**
+ * Movie scene animation embedded within an actor.
+ */
+UCLASS(Blueprintable, Experimental, ClassGroup=Sequence, hidecategories=(Collision, Cooking, Activation), meta=(BlueprintSpawnableComponent))
+class ACTORSEQUENCE_API UActorSequenceComponent
+	: public UActorComponent
+{
+public:
+	GENERATED_BODY()
+
+	UActorSequenceComponent(const FObjectInitializer& Init);
+
+	UActorSequence* GetSequence() const
+	{
+		return Sequence;
+	}
+	
+	virtual void PostInitProperties() override;
+	virtual void BeginPlay() override;
+	virtual void TickComponent(float DeltaTime, enum ELevelTick TickType, FActorComponentTickFunction *ThisTickFunction) override;
+
+protected:
+
+	UPROPERTY(EditAnywhere, Category="Playback", meta=(ShowOnlyInnerProperties))
+	FMovieSceneSequencePlaybackSettings PlaybackSettings;
+
+	/** Embedded actor sequence data */
+	UPROPERTY(EditAnywhere, Instanced, Category=Animation)
+	UActorSequence* Sequence;
+
+	UPROPERTY(transient, BlueprintReadOnly, Category=Animation)
+	UActorSequencePlayer* SequencePlayer;
+
+	UPROPERTY(EditAnywhere, Category="Playback")
+	bool bAutoPlay;
+};
diff --git a/Engine/Source/Runtime/ActorSequence/Public/ActorSequencePlayer.h b/Engine/Source/Runtime/ActorSequence/Public/ActorSequencePlayer.h
new file mode 100644
--- /dev/null
+++ b/Engine/Source/Runtime/ActorSequence/Public/ActorSequencePlayer.h
@@ -0,0 +1,31 @@
+// Copyright 1998-2016 Epic Games, Inc. All Rights Reserved.
+
+#pragma once
+
+#include "IMovieScenePlayer.h"
+#include "ActorSequence.h"
+#include "MovieSceneSequencePlayer.h"
+#include "ActorSequencePlayer.generated.h"
+
+/**
+ * UActorSequencePlayer is used to actually "play" an actor sequence asset at runtime.
+ */
+UCLASS(BlueprintType)
+class ACTORSEQUENCE_API UActorSequencePlayer
+	: public UMovieSceneSequencePlayer
+{
+public:
+	GENERATED_BODY()
+
+protected:
+
+	//~ IMovieScenePlayer interface
+	virtual UObject* GetPlaybackContext() const override;
+	virtual TArray<UObject*> GetEventContexts() const override;
+
+private:
+
+	/** Add tick prerequisites so that the level sequence actor ticks before all the actors it controls */
+	void SetTickPrerequisites(bool bAddTickPrerequisites);
+};
+
diff --git a/Engine/Source/Runtime/AnimGraphRuntime/Private/AnimNodes/AnimNode_TwoWayBlend.cpp b/Engine/Source/Runtime/AnimGraphRuntime/Private/AnimNodes/AnimNode_TwoWayBlend.cpp
--- a/Engine/Source/Runtime/AnimGraphRuntime/Private/AnimNodes/AnimNode_TwoWayBlend.cpp
+++ b/Engine/Source/Runtime/AnimGraphRuntime/Private/AnimNodes/AnimNode_TwoWayBlend.cpp
@@ -4,9 +4,9 @@
 #include "AnimationRuntime.h"
 
 /////////////////////////////////////////////////////
-// FAnimationNode_TwoWayBlend
+// FAnimNode_TwoWayBlend
 
-void FAnimationNode_TwoWayBlend::Initialize(const FAnimationInitializeContext& Context)
+void FAnimNode_TwoWayBlend::Initialize(const FAnimationInitializeContext& Context)
 {
 	FAnimNode_Base::Initialize(Context);
 
@@ -17,13 +17,13 @@ void FAnimationNode_TwoWayBlend::Initialize(const FAnimationInitializeContext& C
 	bBIsRelevant = false;
 }
 
-void FAnimationNode_TwoWayBlend::CacheBones(const FAnimationCacheBonesContext& Context) 
+void FAnimNode_TwoWayBlend::CacheBones(const FAnimationCacheBonesContext& Context) 
 {
 	A.CacheBones(Context);
 	B.CacheBones(Context);
 }
 
-void FAnimationNode_TwoWayBlend::Update(const FAnimationUpdateContext& Context)
+void FAnimNode_TwoWayBlend::Update(const FAnimationUpdateContext& Context)
 {
 	QUICK_SCOPE_CYCLE_COUNTER(STAT_FAnimationNode_TwoWayBlend_Update);
 	EvaluateGraphExposedInputs.Execute(Context);
@@ -76,7 +76,7 @@ void FAnimationNode_TwoWayBlend::Update(const FAnimationUpdateContext& Context)
 	}
 }
 
-void FAnimationNode_TwoWayBlend::Evaluate(FPoseContext& Output)
+void FAnimNode_TwoWayBlend::Evaluate(FPoseContext& Output)
 {
 	if (bBIsRelevant)
 	{
@@ -102,7 +102,7 @@ void FAnimationNode_TwoWayBlend::Evaluate(FPoseContext& Output)
 }
 
 
-void FAnimationNode_TwoWayBlend::GatherDebugData(FNodeDebugData& DebugData)
+void FAnimNode_TwoWayBlend::GatherDebugData(FNodeDebugData& DebugData)
 {
 	FString DebugLine = DebugData.GetNodeName(this);
 	DebugLine += FString::Printf(TEXT("(Alpha: %.1f%%)"), InternalBlendAlpha *100);
diff --git a/Engine/Source/Runtime/AnimGraphRuntime/Private/AnimSequencerInstanceProxy.cpp b/Engine/Source/Runtime/AnimGraphRuntime/Private/AnimSequencerInstanceProxy.cpp
new file mode 100644
--- /dev/null
+++ b/Engine/Source/Runtime/AnimGraphRuntime/Private/AnimSequencerInstanceProxy.cpp
@@ -0,0 +1,147 @@
+// Copyright 1998-2016 Epic Games, Inc. All Rights Reserved.
+
+#include "AnimSequencerInstanceProxy.h"
+#include "AnimSequencerInstance.h"
+
+void FAnimSequencerInstanceProxy::Initialize(UAnimInstance* InAnimInstance)
+{
+	FAnimInstanceProxy::Initialize(InAnimInstance);
+	ConstructNodes();
+}
+
+bool FAnimSequencerInstanceProxy::Evaluate(FPoseContext& Output)
+{
+	SequencerRootNode.Evaluate(Output);
+
+	return true;
+}
+
+void FAnimSequencerInstanceProxy::UpdateAnimationNode(float DeltaSeconds)
+{
+	SequencerRootNode.Update(FAnimationUpdateContext(this, DeltaSeconds));
+}
+
+void FAnimSequencerInstanceProxy::ConstructNodes()
+{
+	// construct node link node for full body and additive to apply additive node
+	SequencerRootNode.Base.SetLinkNode(&FullBodyBlendNode);
+	SequencerRootNode.Additive.SetLinkNode(&AdditiveBlendNode);
+
+	FullBodyBlendNode.bAdditiveNode = false;
+	FullBodyBlendNode.bNormalizeAlpha = true;
+
+	AdditiveBlendNode.bAdditiveNode = true;
+	AdditiveBlendNode.bNormalizeAlpha = false;
+
+	FullBodyBlendNode.ResetPoses();
+	AdditiveBlendNode.ResetPoses();
+
+	ClearSequencePlayerMap();
+}
+
+void FAnimSequencerInstanceProxy::ClearSequencePlayerMap()
+{
+	for (TPair<int32, FSequencerPlayerState*>& Iter : SequencerToPlayerMap)
+	{
+		delete Iter.Value;
+	}
+
+	SequencerToPlayerMap.Empty();
+}
+
+void FAnimSequencerInstanceProxy::ResetNodes()
+{
+	FMemory::Memzero(FullBodyBlendNode.DesiredAlphas.GetData(), FullBodyBlendNode.DesiredAlphas.GetAllocatedSize());
+	FMemory::Memzero(AdditiveBlendNode.DesiredAlphas.GetData(), AdditiveBlendNode.DesiredAlphas.GetAllocatedSize());
+}
+
+FAnimSequencerInstanceProxy::~FAnimSequencerInstanceProxy()
+{
+	ClearSequencePlayerMap();
+}
+
+void FAnimSequencerInstanceProxy::InitAnimTrack(UAnimSequenceBase* InAnimSequence, int32 SequenceId)
+{
+	if (InAnimSequence != nullptr)
+	{
+		FSequencerPlayerState* PlayerState = FindPlayer(SequenceId);
+		if (PlayerState == nullptr)
+		{
+			const bool bIsAdditive = InAnimSequence->IsValidAdditive();
+			FAnimNode_MultiWayBlend& BlendNode = (bIsAdditive) ? AdditiveBlendNode : FullBodyBlendNode;
+			
+			// you shouldn't allow additive animation to be added here, but if it changes type after
+			// you'll see this warning coming up
+			if (bIsAdditive && InAnimSequence->GetAdditiveAnimType() == AAT_RotationOffsetMeshSpace)
+			{
+				// this doesn't work
+				UE_LOG(LogAnimation, Warning, TEXT("ERROR: Animation [%s] in Sequencer has Mesh Space additive animation.  No support on mesh space additive animation. "), *GetNameSafe(InAnimSequence));
+			}
+
+			const int32 PoseIndex = BlendNode.AddPose() - 1;
+
+			// add the new entry to map
+			FSequencerPlayerState* NewPlayerState = new FSequencerPlayerState();
+			NewPlayerState->PoseIndex = PoseIndex;
+			NewPlayerState->bAdditive = bIsAdditive;
+			
+			SequencerToPlayerMap.Add(SequenceId, NewPlayerState);
+
+			// link player to blendnode, this will let you trigger notifies and so on
+			NewPlayerState->PlayerNode.bTeleportToExplicitTime = false;
+			BlendNode.Poses[PoseIndex].SetLinkNode(&NewPlayerState->PlayerNode);
+
+			// set player state
+			PlayerState = NewPlayerState;
+		}
+
+		// now set animation data to player
+		PlayerState->PlayerNode.Sequence = InAnimSequence;
+		PlayerState->PlayerNode.ExplicitTime = 0.f;
+
+		// initialize player
+		PlayerState->PlayerNode.Initialize(FAnimationInitializeContext(this));
+	}
+}
+
+/*
+// this isn't used yet. If we want to optimize it, we could do this way, but right now the way sequencer updates, we don't have a good point 
+// where we could just clear one sequence id. We just clear all the weights before update. 
+// once they go out of range, they don't get called anymore, so there is no good point of tearing down
+// there is multiple tear down point but we couldn't find where only happens once activated and once getting out
+// because sequencer finds the nearest point, not exact point, it doens't have good point of tearing down
+void FAnimSequencerInstanceProxy::TermAnimTrack(int32 SequenceId)
+{
+	FSequencerPlayerState* PlayerState = FindPlayer(SequenceId);
+
+	if (PlayerState)
+	{
+		FAnimNode_MultiWayBlend& BlendNode = (PlayerState->bAdditive) ? AdditiveBlendNode : FullBodyBlendNode;
+
+		// remove the pose from blend node
+		BlendNode.Poses.RemoveAt(PlayerState->PoseIndex);
+		BlendNode.DesiredAlphas.RemoveAt(PlayerState->PoseIndex);
+
+		// remove from Sequence Map
+		SequencerToPlayerMap.Remove(SequenceId);
+	}
+}*/
+
+void FAnimSequencerInstanceProxy::UpdateAnimTrack(UAnimSequenceBase* InAnimSequence, int32 SequenceId, float InPosition, float Weight, bool bFireNotifies)
+{
+	EnsureAnimTrack(InAnimSequence, SequenceId);
+
+	FSequencerPlayerState* PlayerState = FindPlayer(SequenceId);
+	PlayerState->PlayerNode.ExplicitTime = InPosition;
+	FAnimNode_MultiWayBlend& BlendNode = (PlayerState->bAdditive) ? AdditiveBlendNode : FullBodyBlendNode;
+	BlendNode.DesiredAlphas[PlayerState->PoseIndex] = Weight;
+}
+
+void FAnimSequencerInstanceProxy::EnsureAnimTrack(UAnimSequenceBase* InAnimSequence, int32 SequenceId)
+{
+	FSequencerPlayerState* PlayerState = FindPlayer(SequenceId);
+	if (!PlayerState)
+	{
+		InitAnimTrack(InAnimSequence, SequenceId);
+	}
+}
\ No newline at end of file
diff --git a/Engine/Source/Runtime/AnimGraphRuntime/Public/AnimNodes/AnimNode_MultiWayBlend.h b/Engine/Source/Runtime/AnimGraphRuntime/Public/AnimNodes/AnimNode_MultiWayBlend.h
new file mode 100644
--- /dev/null
+++ b/Engine/Source/Runtime/AnimGraphRuntime/Public/AnimNodes/AnimNode_MultiWayBlend.h
@@ -0,0 +1,83 @@
+// Copyright 1998-2016 Epic Games, Inc. All Rights Reserved.
+
+#pragma once
+
+#include "Animation/AnimNodeBase.h"
+#include "Animation/InputScaleBias.h"
+#include "AnimNode_MultiWayBlend.generated.h"
+
+// This represents a baked transition
+USTRUCT()
+struct ANIMGRAPHRUNTIME_API FAnimNode_MultiWayBlend : public FAnimNode_Base
+{
+	GENERATED_USTRUCT_BODY()
+public:
+
+	UPROPERTY(EditAnywhere, BlueprintReadWrite, Category = Links)
+	TArray<FPoseLink> Poses;
+
+	UPROPERTY(EditAnywhere, BlueprintReadWrite, Category=Settings, meta=(PinShownByDefault))
+	TArray<float> DesiredAlphas;
+
+	UPROPERTY(EditAnywhere, BlueprintReadWrite, Category = Settings)
+	bool bAdditiveNode;
+
+	UPROPERTY(EditAnywhere, BlueprintReadWrite, Category = Settings)
+	bool bNormalizeAlpha;
+
+	UPROPERTY(EditAnywhere, BlueprintReadWrite, Category=Settings)
+	FInputScaleBias AlphaScaleBias;
+
+public:
+	FAnimNode_MultiWayBlend()
+		: bAdditiveNode(false)
+		, bNormalizeAlpha(true)
+	{
+	}
+
+	// FAnimNode_Base interface
+	virtual void Initialize(const FAnimationInitializeContext& Context) override;
+	virtual void CacheBones(const FAnimationCacheBonesContext& Context) override;
+	virtual void Update(const FAnimationUpdateContext& Context) override;
+	virtual void Evaluate(FPoseContext& Output) override;
+	virtual void GatherDebugData(FNodeDebugData& DebugData) override;
+	// End of FAnimNode_Base interface
+
+	int32 AddPose()
+	{
+		Poses.AddZeroed();
+		DesiredAlphas.AddZeroed();
+
+		return Poses.Num();
+	}
+
+	void RemovePose(int32 PoseIndex)
+	{
+		Poses.RemoveAt(PoseIndex);
+		DesiredAlphas.RemoveAt(PoseIndex);
+	}
+
+	void ResetPoses()
+	{
+		Poses.Reset();
+		DesiredAlphas.Reset();
+	}
+private:
+	float GetTotalAlpha()
+	{
+		float TotalAlpha = 0.f;
+
+		for (float Alpha : DesiredAlphas)
+		{
+			TotalAlpha += Alpha;
+		}
+
+		return TotalAlpha;
+	}
+
+	// process new weights and then return out
+	void UpdateCachedAlphas();
+
+	TArray<float> CachedAlphas;
+};
+
diff --git a/Engine/Source/Runtime/AnimGraphRuntime/Public/AnimNodes/AnimNode_TwoWayBlend.h b/Engine/Source/Runtime/AnimGraphRuntime/Public/AnimNodes/AnimNode_TwoWayBlend.h
--- a/Engine/Source/Runtime/AnimGraphRuntime/Public/AnimNodes/AnimNode_TwoWayBlend.h
+++ b/Engine/Source/Runtime/AnimGraphRuntime/Public/AnimNodes/AnimNode_TwoWayBlend.h
@@ -10,7 +10,7 @@
 
 // This represents a baked transition
 USTRUCT()
-struct ANIMGRAPHRUNTIME_API FAnimationNode_TwoWayBlend : public FAnimNode_Base
+struct ANIMGRAPHRUNTIME_API FAnimNode_TwoWayBlend : public FAnimNode_Base
 {
 	GENERATED_USTRUCT_BODY()
 public:
@@ -42,7 +42,7 @@ protected:
 	bool bResetChildOnActivation;
 
 public:
-	FAnimationNode_TwoWayBlend()
+	FAnimNode_TwoWayBlend()
 		: Alpha(0.0f)
 		, bResetChildOnActivation(false)
 	{
diff --git a/Engine/Source/Runtime/AnimGraphRuntime/Public/AnimSequencerInstance.h b/Engine/Source/Runtime/AnimGraphRuntime/Public/AnimSequencerInstance.h
new file mode 100644
--- /dev/null
+++ b/Engine/Source/Runtime/AnimGraphRuntime/Public/AnimSequencerInstance.h
@@ -0,0 +1,26 @@
+// Copyright 1998-2016 Epic Games, Inc. All Rights Reserved.
+
+/**
+ *
+ * This Instance only contains one AnimationAsset, and produce poses
+ * Used by Preview in AnimGraph, Playing single animation in Kismet2 and etc
+ */
+
+#pragma once
+#include "Animation/AnimInstance.h"
+#include "AnimSequencerInstance.generated.h"
+
+UCLASS(transient, NotBlueprintable)
+class ANIMGRAPHRUNTIME_API UAnimSequencerInstance : public UAnimInstance
+{
+	GENERATED_UCLASS_BODY()
+
+public:
+	void UpdateAnimTrack(UAnimSequenceBase* InAnimSequence, int32 SequenceId, float InPosition, float Weight, bool bFireNotifies);
+	void ResetNodes();
+protected:
+	virtual FAnimInstanceProxy* CreateAnimInstanceProxy() override;
+};
+
+
+
diff --git a/Engine/Source/Runtime/AnimGraphRuntime/Public/AnimSequencerInstanceProxy.h b/Engine/Source/Runtime/AnimGraphRuntime/Public/AnimSequencerInstanceProxy.h
new file mode 100644
--- /dev/null
+++ b/Engine/Source/Runtime/AnimGraphRuntime/Public/AnimSequencerInstanceProxy.h
@@ -0,0 +1,65 @@
+// Copyright 1998-2016 Epic Games, Inc. All Rights Reserved.
+
+#pragma once
+
+#include "AnimInstanceProxy.h"
+#include "AnimNodes/AnimNode_SequenceEvaluator.h"
+#include "AnimNodes/AnimNode_ApplyAdditive.h"
+#include "AnimNodes/AnimNode_MultiWayBlend.h"
+#include "AnimSequencerInstanceProxy.generated.h"
+
+struct FSequencerPlayerState
+{
+	int32 PoseIndex;
+	bool bAdditive;
+	struct FAnimNode_SequenceEvaluator PlayerNode;
+
+	FSequencerPlayerState()
+		: PoseIndex(INDEX_NONE)
+		, bAdditive(false)
+	{}
+};
+
+/** Proxy override for this UAnimInstance-derived class */
+USTRUCT()
+struct ANIMGRAPHRUNTIME_API FAnimSequencerInstanceProxy : public FAnimInstanceProxy
+{
+	GENERATED_BODY()
+
+public:
+	FAnimSequencerInstanceProxy()
+	{
+	}
+
+	FAnimSequencerInstanceProxy(UAnimInstance* InAnimInstance)
+		: FAnimInstanceProxy(InAnimInstance)
+	{
+	}
+
+	virtual ~FAnimSequencerInstanceProxy();
+
+	virtual void Initialize(UAnimInstance* InAnimInstance) override;
+	virtual bool Evaluate(FPoseContext& Output) override;
+	virtual void UpdateAnimationNode(float DeltaSeconds) override;
+
+	void UpdateAnimTrack(UAnimSequenceBase* InAnimSequence, int32 SequenceId, float InPosition, float Weight, bool bFireNotifies);
+	void ResetNodes();
+private:
+	void ConstructNodes();
+	struct FSequencerPlayerState* FindPlayer(int32 SequenceId)
+	{
+		return SequencerToPlayerMap.FindRef(SequenceId);
+	}
+
+	/** custom root node for this sequencer player. Didn't want to use RootNode in AnimInstance because it's used with lots of AnimBP functionality */
+	struct FAnimNode_ApplyAdditive SequencerRootNode;
+	struct FAnimNode_MultiWayBlend FullBodyBlendNode;
+	struct FAnimNode_MultiWayBlend AdditiveBlendNode;
+
+	/** mapping from sequencer index to internal player index */
+	TMap<int32, FSequencerPlayerState*> SequencerToPlayerMap;
+
+	void InitAnimTrack(UAnimSequenceBase* InAnimSequence, int32 SequenceId);
+	void EnsureAnimTrack(UAnimSequenceBase* InAnimSequence, int32 SequenceId);
+	void ClearSequencePlayerMap();
+};
\ No newline at end of file
diff --git a/Engine/Source/Runtime/CinematicCamera/Private/CameraRig_Rail.cpp b/Engine/Source/Runtime/CinematicCamera/Private/CameraRig_Rail.cpp
--- a/Engine/Source/Runtime/CinematicCamera/Private/CameraRig_Rail.cpp
+++ b/Engine/Source/Runtime/CinematicCamera/Private/CameraRig_Rail.cpp
@@ -127,7 +127,7 @@ void ACameraRig_Rail::UpdatePreviewMeshes()
 			GetComponents(OwnedSplineMeshComponents);
 			for (auto OwnedComponent : OwnedSplineMeshComponents)
 			{
-				if (!PreviewRailMeshSegments.Contains(OwnedComponent))
+				if (!PreviewRailMeshSegments.Contains(OwnedComponent) && OwnedComponent->IsRegistered())
 				{
 					OwnedComponent->UnregisterComponent();
 				}
diff --git a/Engine/Source/Runtime/CinematicCamera/Private/CineCameraComponent.cpp b/Engine/Source/Runtime/CinematicCamera/Private/CineCameraComponent.cpp
--- a/Engine/Source/Runtime/CinematicCamera/Private/CineCameraComponent.cpp
+++ b/Engine/Source/Runtime/CinematicCamera/Private/CineCameraComponent.cpp
@@ -113,6 +113,9 @@ void UCineCameraComponent::PostEditChangeProperty(FPropertyChangedEvent& Propert
 		DebugFocusPlaneMID->SetVectorParameterValue(FName(TEXT("Color")), FocusSettings.DebugFocusPlaneColor.ReinterpretAsLinear());
 	}
 
+	// reset interpolation if the user changes anything
+	bResetInterpolation = true;
+
 	Super::PostEditChangeProperty(PropertyChangedEvent);
 }
 
diff --git a/Engine/Source/Runtime/Core/Public/Containers/Algo/Find.h b/Engine/Source/Runtime/Core/Public/Containers/Algo/Find.h
--- a/Engine/Source/Runtime/Core/Public/Containers/Algo/Find.h
+++ b/Engine/Source/Runtime/Core/Public/Containers/Algo/Find.h
@@ -3,11 +3,12 @@
 #pragma once
 
 #include "CoreTypes.h"
+#include "Traits/ElementType.h"
 
 namespace AlgoImpl
 {
 	template <typename RangeType, typename PredicateType>
-	const typename RangeType::ElementType* FindByPredicate(const RangeType& Range, PredicateType Predicate)
+	const typename TElementType<RangeType>::Type* FindByPredicate(const RangeType& Range, PredicateType Predicate)
 	{
 		for (const auto& Elem : Range)
 		{
@@ -31,9 +32,9 @@ namespace Algo
 	 * @return A pointer to the first element found, or nullptr if none was found.
 	 */
 	template <typename RangeType, typename ValueType>
-	FORCEINLINE typename RangeType::ElementType* Find(RangeType& Range, const ValueType& Value)
+	FORCEINLINE typename TElementType<RangeType>::Type* Find(RangeType& Range, const ValueType& Value)
 	{
-		return const_cast<typename RangeType::ElementType*>(Find(const_cast<const RangeType&>(Range), Value));
+		return const_cast<typename TElementType<RangeType>::Type*>(Find(const_cast<const RangeType&>(Range), Value));
 	}
 
 	/**
@@ -44,9 +45,9 @@ namespace Algo
 	 * @return A pointer to the first element found, or nullptr if none was found.
 	 */
 	template <typename RangeType, typename ValueType>
-	FORCEINLINE const typename RangeType::ElementType* Find(const RangeType& Range, const ValueType& Value)
+	FORCEINLINE const typename TElementType<RangeType>::Type* Find(const RangeType& Range, const ValueType& Value)
 	{
-		return AlgoImpl::FindByPredicate(Range, [&](const typename RangeType::ElementType& Elem){
+		return AlgoImpl::FindByPredicate(Range, [&](const typename TElementType<RangeType>::Type& Elem){
 			return Elem == Value;
 		});
 	}
@@ -59,9 +60,9 @@ namespace Algo
 	 * @return A pointer to the first element found, or nullptr if none was found.
 	 */
 	template <typename RangeType, typename PredicateType>
-	FORCEINLINE typename RangeType::ElementType* FindByPredicate(RangeType& Range, PredicateType Pred)
+	FORCEINLINE typename TElementType<RangeType>::Type* FindByPredicate(RangeType& Range, PredicateType Pred)
 	{
-		return const_cast<typename RangeType::ElementType*>(FindByPredicate(const_cast<const RangeType&>(Range), Pred));
+		return const_cast<typename TElementType<RangeType>::Type*>(FindByPredicate(const_cast<const RangeType&>(Range), Pred));
 	}
 
 	/**
@@ -72,7 +73,7 @@ namespace Algo
 	 * @return A pointer to the first element found, or nullptr if none was found.
 	 */
 	template <typename RangeType, typename PredicateType>
-	FORCEINLINE const typename RangeType::ElementType* FindByPredicate(const RangeType& Range, PredicateType Pred)
+	FORCEINLINE const typename TElementType<RangeType>::Type* FindByPredicate(const RangeType& Range, PredicateType Pred)
 	{
 		return AlgoImpl::FindByPredicate(Range, Pred);
 	}
diff --git a/Engine/Source/Runtime/Core/Public/Misc/Timespan.h b/Engine/Source/Runtime/Core/Public/Misc/Timespan.h
--- a/Engine/Source/Runtime/Core/Public/Misc/Timespan.h
+++ b/Engine/Source/Runtime/Core/Public/Misc/Timespan.h
@@ -122,7 +122,6 @@ public:
 	FTimespan& operator+=(const FTimespan& Other)
 	{
 		Ticks += Other.Ticks;
-
 		return *this;
 	}
 
@@ -158,7 +157,6 @@ public:
 	FTimespan& operator-=(const FTimespan& Other)
 	{
 		Ticks -= Other.Ticks;
-
 		return *this;
 	}
 
@@ -182,7 +180,29 @@ public:
 	FTimespan& operator*=(float Scalar)
 	{
 		Ticks = (int64)(Ticks * Scalar);
+		return *this;
+	}
+
+	/**
+	 * Returns the result of calculating the modulus of this time span with another time span.
+	 *
+	 * @param Other The time span to divide by.
+	 * @return A time span representing the remainder of the modulus operation.
+	 */
+	FTimespan operator%(const FTimespan& Other) const
+	{
+		return FTimespan(Ticks % Other.Ticks);
+	}
 
+	/**
+	 * Calculates this time span modulo another.
+	 *
+	 * @param Other The time span to divide by.
+	 * @return This time span.
+	 */
+	FTimespan& operator%=(const FTimespan& Other)
+	{
+		Ticks = Ticks % Other.Ticks;
 		return *this;
 	}
 
@@ -433,6 +453,16 @@ public:
 	 */
 	CORE_API bool ImportTextItem(const TCHAR*& Buffer, int32 PortFlags, UObject* Parent, FOutputDevice* ErrorText);
 
+	/**
+	 * Check whether this time span is zero.
+	 *
+	 * @return true if the time span is zero, false otherwise.
+	 * @see Zero
+	 */
+	bool IsZero() const
+	{
+		return (Ticks == 0LL);
+	}
 
 	/**
 	 * Serializes this time span from or into the specified archive.
@@ -577,7 +607,7 @@ public:
 	 * The zero time span value can be used in comparison operations with other time spans.
 	 *
 	 * @return Zero time span.
-	 * @see MaxValue, MinValue
+	 * @see IsZero, MaxValue, MinValue
 	 */
 	static FTimespan Zero()
 	{
diff --git a/Engine/Source/Runtime/Core/Public/UObject/SequencerObjectVersion.h b/Engine/Source/Runtime/Core/Public/UObject/SequencerObjectVersion.h
--- a/Engine/Source/Runtime/Core/Public/UObject/SequencerObjectVersion.h
+++ b/Engine/Source/Runtime/Core/Public/UObject/SequencerObjectVersion.h
@@ -15,6 +15,12 @@ struct CORE_API FSequencerObjectVersion
 		// Per-platform overrides player overrides for media sources changed name and type.
 		RenameMediaSourcePlatformPlayers,
 
+		// Enable root motion isn't the right flag to use, but force root lock
+		ConvertEnableRootMotionToForceRootLock,
+
+		// Convert multiple rows to tracks
+		ConvertMultipleRowsToTracks,
+
 		// -----<new versions can be added above this line>-------------------------------------------------
 		VersionPlusOne,
 		LatestVersion = VersionPlusOne - 1
diff --git a/Engine/Source/Runtime/Engine/Classes/Components/SkeletalMeshComponent.h b/Engine/Source/Runtime/Engine/Classes/Components/SkeletalMeshComponent.h
--- a/Engine/Source/Runtime/Engine/Classes/Components/SkeletalMeshComponent.h
+++ b/Engine/Source/Runtime/Engine/Classes/Components/SkeletalMeshComponent.h
@@ -252,7 +252,9 @@ namespace EAnimationMode
 	enum Type
 	{
 		AnimationBlueprint UMETA(DisplayName="Use Animation Blueprint"), 
-		AnimationSingleNode UMETA(DisplayName="Use Animation Asset")
+		AnimationSingleNode UMETA(DisplayName="Use Animation Asset"), 
+		// This is custom type, engine leaves AnimInstance as it is
+		AnimationCustomMode UMETA(DisplayName = "Use Custom Mode"),
 	};
 }
 
diff --git a/Engine/Source/Runtime/Engine/Engine.Build.cs b/Engine/Source/Runtime/Engine/Engine.Build.cs
--- a/Engine/Source/Runtime/Engine/Engine.Build.cs
+++ b/Engine/Source/Runtime/Engine/Engine.Build.cs
@@ -167,6 +167,7 @@ public class Engine : ModuleRules
 		DynamicallyLoadedModuleNames.AddRange(
 			new string[]
 			{
+				"ActorSequence",
 				"MovieScene",
 				"MovieSceneCapture",
 				"MovieSceneTracks",
diff --git a/Engine/Source/Runtime/Engine/Private/Components/SkeletalMeshComponent.cpp b/Engine/Source/Runtime/Engine/Private/Components/SkeletalMeshComponent.cpp
--- a/Engine/Source/Runtime/Engine/Private/Components/SkeletalMeshComponent.cpp
+++ b/Engine/Source/Runtime/Engine/Private/Components/SkeletalMeshComponent.cpp
@@ -2632,10 +2632,15 @@ void USkeletalMeshComponent::ValidateAnimation()
 		{
 			AnimationData.AnimToPlay = nullptr;
 		}
-		else
+		else if(AnimationMode == EAnimationMode::AnimationBlueprint)
 		{
 			AnimClass = nullptr;
 		}
+		else
+		{
+			// if custom mode, you still can't use the animation instance
+			AnimScriptInstance = nullptr;
+		}
 		return;
 	}
 
@@ -2655,7 +2660,7 @@ void USkeletalMeshComponent::ValidateAnimation()
 			AnimationData.AnimToPlay = nullptr;
 		}
 	}
-	else
+	else if (AnimationMode == EAnimationMode::AnimationBlueprint)
 	{
 		IAnimClassInterface* AnimClassInterface = IAnimClassInterface::GetFromClass(AnimClass);
 		if (AnimClassInterface && SkeletalMesh && AnimClassInterface->GetTargetSkeleton() != SkeletalMesh->Skeleton)
diff --git a/Engine/Source/Runtime/Engine/Private/Level.cpp b/Engine/Source/Runtime/Engine/Private/Level.cpp
--- a/Engine/Source/Runtime/Engine/Private/Level.cpp
+++ b/Engine/Source/Runtime/Engine/Private/Level.cpp
@@ -492,8 +492,8 @@ void ULevel::SortActorList()
 				NewActors.Add(Actor);
 			}
 		}
-
 	}
+
 	iFirstNetRelevantActor = NewActors.Num();
 
 	NewActors.Append(MoveTemp(NewNetActors));
@@ -510,6 +510,11 @@ void ULevel::SortActorList()
 		{
 			iFirstNetRelevantActor = 0;
 		}
+		// Ensure the world settings actor is added if it's not going to get added in the loop below
+		else if ( IsNetActor( WorldSettings ) )
+		{
+			OwningWorld->AddNetworkActor( WorldSettings );
+		}
 
 		for ( int32 i = iFirstNetRelevantActor; i < Actors.Num(); i++ )
 		{
diff --git a/Engine/Source/Runtime/Engine/Private/Materials/MaterialShared.cpp b/Engine/Source/Runtime/Engine/Private/Materials/MaterialShared.cpp
--- a/Engine/Source/Runtime/Engine/Private/Materials/MaterialShared.cpp
+++ b/Engine/Source/Runtime/Engine/Private/Materials/MaterialShared.cpp
@@ -1726,7 +1726,7 @@ bool FMaterial::BeginCompileShaderMap(
 		else
 		{
 #if DEBUG_INFINITESHADERCOMPILE
-			UE_LOG(LogTemp, Display, TEXT("Kicking of shader compilation for %s, GameThreadShaderMap 0x%08X%08X"), *GetFriendlyName(), (int)((int64)(NewShaderMap.GetReference()) >> 32), (int)((int64)(NewShaderMap.GetReference())));
+			UE_LOG(LogTemp, Display, TEXT("Kicking off shader compilation for %s, GameThreadShaderMap 0x%08X%08X"), *GetFriendlyName(), (int)((int64)(NewShaderMap.GetReference()) >> 32), (int)((int64)(NewShaderMap.GetReference())));
 #endif
 			OutstandingCompileShaderMapIds.AddUnique( NewShaderMap->GetCompilingId() );
 			// Async compile, use NULL so that rendering will fall back to the default material.
diff --git a/Engine/Source/Runtime/Engine/Public/Curves/CurveInterface.h b/Engine/Source/Runtime/Engine/Public/Curves/CurveInterface.h
--- a/Engine/Source/Runtime/Engine/Public/Curves/CurveInterface.h
+++ b/Engine/Source/Runtime/Engine/Public/Curves/CurveInterface.h
@@ -17,6 +17,9 @@ struct TKeyFrameProxy
 		, Value(InValue)
 	{}
 
+	TKeyFrameProxy* operator->() { return this; }
+	const TKeyFrameProxy* operator->() const { return this; }
+
 	/** The key time */
 	TimeType Time;
 
@@ -24,6 +27,8 @@ struct TKeyFrameProxy
 	KeyValueType& Value;
 };
 
+template<typename KeyValueType, typename TimeType> struct TKeyIterator;
+
 /**
  * Templated interface to externally owned curve data data
  */
@@ -78,7 +83,7 @@ public:
 		int32 Index = this->GetIndex(KeyHandle);
 		if (Index != INDEX_NONE)
 		{
-			return TKeyFrameProxy<const KeyValueType, const TimeType>(this->GetKeyTimeChecked(Index), (*KeyValues)[Index]);
+			return TKeyFrameProxy<const KeyValueType, TimeType>(this->GetKeyTimeChecked(Index), (*KeyValues)[Index]);
 		}
 		return TOptional<TKeyFrameProxy<const KeyValueType, TimeType>>();
 	}
@@ -132,6 +137,24 @@ public:
 		}
 	}
 
+public:
+
+	/**
+	 * Iterate this curve's keys
+	 */
+	TKeyIterator<KeyValueType, TimeType> IterateKeysAndValues()
+	{
+		return TKeyIterator<KeyValueType, TimeType>(*this);
+	}
+
+	/**
+	 * Iterate this curve's keys
+	 */
+	TKeyIterator<const KeyValueType, TimeType> IterateKeysAndValues() const
+	{
+		return TKeyIterator<const KeyValueType, TimeType>(*this);
+	}
+
 protected:
 
 	virtual void OnKeyAdded(int32 Index) override
@@ -155,6 +178,108 @@ protected:
 
 private:
 
+	friend TKeyIterator<KeyValueType, TimeType>;
+	friend TKeyIterator<const KeyValueType, TimeType>;
+
 	/** Array of associated key data */
 	TArray<KeyValueType>* KeyValues;
 };
+
+
+namespace Impl
+{
+	template<typename TSrc, typename TDest>
+	struct TMatchConst
+	{
+		typedef TDest Type;
+	};
+
+	template<typename TSrc, typename TDest>
+	struct TMatchConst<const TSrc, TDest>
+	{
+		typedef const TDest Type;
+	};
+}
+
+/**
+ * Key time iterator for TKeyFrameManipulator
+ */
+template<typename KeyValueType, typename TimeType>
+struct TKeyIterator
+{
+	typedef typename Impl::TMatchConst<KeyValueType, TCurveInterface<typename TRemoveConst<KeyValueType>::Type, TimeType>>::Type CurveInterfaceType;
+
+	TKeyIterator(CurveInterfaceType& InCurveInterface)
+		: CurveInterface(InCurveInterface)
+		, Index(0)
+	{}
+
+	FORCEINLINE TKeyIterator& operator++()
+	{
+		++Index;
+		return *this;
+	}
+
+	/** conversion to "bool" returning true if the iterator is valid. */
+	FORCEINLINE explicit operator bool() const
+	{ 
+		return Index < CurveInterface.KeyValues->Num();
+	}
+	/** inverse of the "bool" operator */
+	FORCEINLINE bool operator !() const 
+	{
+		return !(bool)*this;
+	}
+
+	FORCEINLINE friend bool operator==(const TKeyIterator& LHS, const TKeyIterator& RHS)
+	{
+		return &LHS.CurveInterface == &RHS.CurveInterface && LHS.Index == RHS.Index;
+	}
+
+	FORCEINLINE friend bool operator!=(const TKeyIterator& LHS, const TKeyIterator& RHS)
+	{
+		return !(LHS == RHS);
+	}
+
+	FORCEINLINE TKeyFrameProxy<KeyValueType, TimeType> operator*() const
+	{
+		return TKeyFrameProxy<KeyValueType, TimeType>(CurveInterface.GetKeyTimeChecked(Index), (*CurveInterface.KeyValues)[Index]);
+	}
+
+	FORCEINLINE TKeyFrameProxy<KeyValueType, TimeType> operator->() const
+	{
+		return TKeyFrameProxy<KeyValueType, TimeType>(CurveInterface.GetKeyTimeChecked(Index), (*CurveInterface.KeyValues)[Index]);
+	}
+
+	FKeyHandle GetKeyHandle()
+	{
+		return CurveInterface.GetKeyHandleFromIndex(Index);
+	}
+
+	int32 GetStartIndex() const
+	{
+		return 0;
+	}
+
+	int32 GetEndIndex() const
+	{
+		return CurveInterface.KeyValues->Num();
+	}
+
+	FORCEINLINE friend TKeyIterator<KeyValueType, TimeType> begin(const TKeyIterator<KeyValueType, TimeType>& Iter)
+	{
+		TKeyIterator<KeyValueType, TimeType> NewIter = Iter;
+		NewIter.Index = Iter.GetStartIndex();
+		return NewIter;
+	}
+	FORCEINLINE friend TKeyIterator<KeyValueType, TimeType> end(const TKeyIterator<KeyValueType, TimeType>& Iter)
+	{
+		TKeyIterator<KeyValueType, TimeType> NewIter = Iter;
+		NewIter.Index = NewIter.GetEndIndex();
+		return NewIter;
+	}
+
+private:
+	CurveInterfaceType& CurveInterface;
+	int32 Index;
+};
diff --git a/Engine/Source/Runtime/Engine/Public/Curves/KeyFrameManipulator.h b/Engine/Source/Runtime/Engine/Public/Curves/KeyFrameManipulator.h
--- a/Engine/Source/Runtime/Engine/Public/Curves/KeyFrameManipulator.h
+++ b/Engine/Source/Runtime/Engine/Public/Curves/KeyFrameManipulator.h
@@ -194,6 +194,14 @@ protected:
 	{
 		return KeyHandleLUT->GetIndex(KeyHandle);
 	}
+	
+	/**
+	 * Get the handle that corresponds to the specified index
+	 */
+	FKeyHandle GetKeyHandleFromIndex(int32 Index) const
+	{
+		return KeyTimes->IsValidIndex(Index) ? KeyHandleLUT->FindOrAddKeyHandle(Index) : FKeyHandle();
+	}
 
 	/**
 	 * Move a key from one index to another
diff --git a/Engine/Source/Runtime/LevelSequence/LevelSequence.Build.cs b/Engine/Source/Runtime/LevelSequence/LevelSequence.Build.cs
--- a/Engine/Source/Runtime/LevelSequence/LevelSequence.Build.cs
+++ b/Engine/Source/Runtime/LevelSequence/LevelSequence.Build.cs
@@ -9,14 +9,30 @@ public class LevelSequence : ModuleRules
 		PrivateIncludePaths.Add("Runtime/LevelSequence/Private");
 
 		PublicDependencyModuleNames.AddRange(
-            new string[] {
-                "Core",
-                "CoreUObject",
-                "Engine",
+			new string[] {
+				"Core",
+				"CoreUObject",
+				"Engine",
 				"MovieScene",
 				"MovieSceneTracks",
 				"UMG",
-            }
-        );
+			}
+		);
+
+		if (UEBuildConfiguration.bBuildEditor == true)
+		{
+			PrivateDependencyModuleNames.AddRange(
+				new string[] {
+					"ActorPickerMode",
+					"PropertyEditor",
+				}
+			);
+
+			PrivateIncludePathModuleNames.AddRange(
+				new string[] {
+					"SceneOutliner"
+				}
+			);
+		}
 	}
 }
diff --git a/Engine/Source/Runtime/LevelSequence/Private/LevelSequence.cpp b/Engine/Source/Runtime/LevelSequence/Private/LevelSequence.cpp
--- a/Engine/Source/Runtime/LevelSequence/Private/LevelSequence.cpp
+++ b/Engine/Source/Runtime/LevelSequence/Private/LevelSequence.cpp
@@ -5,8 +5,11 @@
 #include "HAL/IConsoleManager.h"
 #include "Components/ActorComponent.h"
 #include "GameFramework/Actor.h"
+#include "Engine/BlueprintGeneratedClass.h"
 #include "Engine/Engine.h"
 #include "MovieScene.h"
+#include "UObject/Package.h"
+#include "UObject/UObjectHash.h"
 
 DEFINE_LOG_CATEGORY_STATIC(LogLevelSequence, Log, All);
 
@@ -167,6 +170,11 @@ bool ULevelSequence::AllowsSpawnableObjects() const
 	return true;
 }
 
+bool ULevelSequence::CanRebindPossessable(const FMovieScenePossessable& InPossessable) const
+{
+	return !InPossessable.GetParent().IsValid();
+}
+
 void ULevelSequence::UnbindPossessableObjects(const FGuid& ObjectId)
 {
 	ObjectReferences.RemoveBinding(ObjectId);
diff --git a/Engine/Source/Runtime/LevelSequence/Private/LevelSequenceActor.cpp b/Engine/Source/Runtime/LevelSequence/Private/LevelSequenceActor.cpp
--- a/Engine/Source/Runtime/LevelSequence/Private/LevelSequenceActor.cpp
+++ b/Engine/Source/Runtime/LevelSequence/Private/LevelSequenceActor.cpp
@@ -6,6 +6,11 @@
 #include "Components/BillboardComponent.h"
 #include "LevelSequenceBurnIn.h"
 
+#if WITH_EDITOR
+	#include "PropertyCustomizationHelpers.h"
+	#include "ActorPickerMode.h"
+	#include "SceneOutlinerFilters.h"
+#endif
 
 ALevelSequenceActor::ALevelSequenceActor(const FObjectInitializer& Init)
 	: Super(Init)
@@ -38,6 +43,7 @@ ALevelSequenceActor::ALevelSequenceActor(const FObjectInitializer& Init)
 	}
 #endif //WITH_EDITORONLY_DATA
 
+	BindingOverrides = Init.CreateDefaultSubobject<UMovieSceneBindingOverrides>(this, "BindingOverrides");
 	BurnInOptions = Init.CreateDefaultSubobject<ULevelSequenceBurnInOptions>(this, "BurnInOptions");
 	PrimaryActorTick.bCanEverTick = true;
 	bAutoPlay = false;
@@ -105,6 +111,40 @@ void ALevelSequenceActor::PostLoad()
 #endif
 }
 
+#if WITH_EDITOR
+
+void FBoundActorProxy::Initialize(TSharedPtr<IPropertyHandle> InPropertyHandle)
+{
+	ReflectedProperty = InPropertyHandle;
+
+	UObject* Object = nullptr;
+	ReflectedProperty->GetValue(Object);
+	BoundActor = Cast<AActor>(Object);
+
+	ReflectedProperty->SetOnPropertyValueChanged(FSimpleDelegate::CreateRaw(this, &FBoundActorProxy::OnReflectedPropertyChanged));
+}
+
+void FBoundActorProxy::OnReflectedPropertyChanged()
+{
+	UObject* Object = nullptr;
+	ReflectedProperty->GetValue(Object);
+	BoundActor = Cast<AActor>(Object);
+}
+
+TSharedPtr<FStructOnScope> ALevelSequenceActor::GetObjectPickerProxy(TSharedPtr<IPropertyHandle> ObjectPropertyHandle)
+{
+	TSharedRef<FStructOnScope> Struct = MakeShared<FStructOnScope>(FBoundActorProxy::StaticStruct());
+	reinterpret_cast<FBoundActorProxy*>(Struct->GetStructMemory())->Initialize(ObjectPropertyHandle);
+	return Struct;
+}
+
+void ALevelSequenceActor::UpdateObjectFromProxy(FStructOnScope& Proxy, IPropertyHandle& ObjectPropertyHandle)
+{
+	UObject* BoundActor = reinterpret_cast<FBoundActorProxy*>(Proxy.GetStructMemory())->BoundActor;
+	ObjectPropertyHandle.SetValue(BoundActor);
+}
+
+#endif
 
 ULevelSequence* ALevelSequenceActor::GetSequence(bool Load) const
 {
@@ -128,6 +168,8 @@ void ALevelSequenceActor::InitializePlayer()
 
 	if (GetWorld()->IsGameWorld() && (LevelSequenceAsset != nullptr))
 	{
+		PlaybackSettings.BindingOverrides = BindingOverrides;
+
 		SequencePlayer = NewObject<ULevelSequencePlayer>(this, "AnimationPlayer");
 		SequencePlayer->Initialize(LevelSequenceAsset, GetWorld(), PlaybackSettings);
 
@@ -191,7 +233,13 @@ void ULevelSequenceBurnInOptions::ResetSettings()
 		{
 			if (!Settings || !Settings->IsA(SettingsClass))
 			{
-				Settings = NewObject<ULevelSequenceBurnInInitSettings>(this, SettingsClass);
+				if (Settings)
+				{
+					Settings->Rename(*MakeUniqueObjectName(this, ULevelSequenceBurnInInitSettings::StaticClass(), "Settings_EXPIRED").ToString());
+				}
+				
+				Settings = NewObject<ULevelSequenceBurnInInitSettings>(this, SettingsClass, "Settings");
+				Settings->SetFlags(GetMaskedFlags(RF_PropagateToSubObjects));
 			}
 		}
 		else
diff --git a/Engine/Source/Runtime/LevelSequence/Private/LevelSequenceBurnIn.cpp b/Engine/Source/Runtime/LevelSequence/Private/LevelSequenceBurnIn.cpp
--- a/Engine/Source/Runtime/LevelSequence/Private/LevelSequenceBurnIn.cpp
+++ b/Engine/Source/Runtime/LevelSequence/Private/LevelSequenceBurnIn.cpp
@@ -17,7 +17,7 @@ void ULevelSequenceBurnIn::TakeSnapshotsFrom(ALevelSequenceActor& InActor)
 	}
 }
 
-void ULevelSequenceBurnIn::OnSequenceUpdated(const ULevelSequencePlayer& Player, float CurrentTime, float PreviousTime)
+void ULevelSequenceBurnIn::OnSequenceUpdated(const UMovieSceneSequencePlayer& Player, float CurrentTime, float PreviousTime)
 {
-	Player.TakeFrameSnapshot(FrameInformation);
+	static_cast<const ULevelSequencePlayer&>(Player).TakeFrameSnapshot(FrameInformation);
 }
diff --git a/Engine/Source/Runtime/LevelSequence/Private/LevelSequencePlayer.cpp b/Engine/Source/Runtime/LevelSequence/Private/LevelSequencePlayer.cpp
--- a/Engine/Source/Runtime/LevelSequence/Private/LevelSequencePlayer.cpp
+++ b/Engine/Source/Runtime/LevelSequence/Private/LevelSequencePlayer.cpp
@@ -18,73 +18,14 @@
 #include "Engine/LevelStreaming.h"
 #include "Tracks/MovieSceneCinematicShotTrack.h"
 #include "Sections/MovieSceneCinematicShotSection.h"
+#include "LevelSequenceActor.h"
 
-struct FTickAnimationPlayers : public FTickableGameObject
-{
-	TArray<TWeakObjectPtr<ULevelSequencePlayer>> ActiveInstances;
-	
-	virtual bool IsTickable() const override
-	{
-		return ActiveInstances.Num() != 0;
-	}
-
-	virtual TStatId GetStatId() const override
-	{
-		RETURN_QUICK_DECLARE_CYCLE_STAT(FTickAnimationPlayers, STATGROUP_Tickables);
-	}
-	
-	virtual void Tick(float DeltaSeconds) override
-	{
-		for (int32 Index = 0; Index < ActiveInstances.Num();)
-		{
-			if (auto* Player = ActiveInstances[Index].Get())
-			{
-				Player->Update(DeltaSeconds);
-				++Index;
-			}
-			else
-			{
-				ActiveInstances.RemoveAt(Index, 1, false);
-			}
-		}
-	}
-};
-
-struct FAutoDestroyAnimationTicker
-{
-	FAutoDestroyAnimationTicker()
-	{
-		FCoreDelegates::OnPreExit.AddLambda([&]{
-			Impl.Reset();
-		});
-	}
-	
-	void Add(ULevelSequencePlayer* Player)
-	{
-		if (!Impl.IsValid())
-		{
-			Impl.Reset(new FTickAnimationPlayers);
-		}
-		Impl->ActiveInstances.Add(Player);
-	}
-
-	TUniquePtr<FTickAnimationPlayers> Impl;
-} GAnimationPlayerTicker;
 
 /* ULevelSequencePlayer structors
  *****************************************************************************/
 
 ULevelSequencePlayer::ULevelSequencePlayer(const FObjectInitializer& ObjectInitializer)
 	: Super(ObjectInitializer)
-	, LevelSequence(nullptr)
-	, bIsPlaying(false)
-	, bReversePlayback(false)
-	, TimeCursorPosition(0.0f)
-	, StartTime(0.f)
-	, EndTime(0.f)
-	, CurrentNumLoops(0)
-	, bHasCleanedUpSequence(false)
-	, bIsEvaluating(false)
 {
 	SpawnRegister = MakeShareable(new FLevelSequenceSpawnRegister);
 }
@@ -93,7 +34,7 @@ ULevelSequencePlayer::ULevelSequencePlayer(const FObjectInitializer& ObjectIniti
 /* ULevelSequencePlayer interface
  *****************************************************************************/
 
-ULevelSequencePlayer* ULevelSequencePlayer::CreateLevelSequencePlayer(UObject* WorldContextObject, ULevelSequence* InLevelSequence, FLevelSequencePlaybackSettings Settings)
+ULevelSequencePlayer* ULevelSequencePlayer::CreateLevelSequencePlayer(UObject* WorldContextObject, ULevelSequence* InLevelSequence, FMovieSceneSequencePlaybackSettings Settings)
 {
 	if (InLevelSequence == nullptr)
 	{
@@ -103,81 +44,18 @@ ULevelSequencePlayer* ULevelSequencePlayer::CreateLevelSequencePlayer(UObject* W
 	UWorld* World = GEngine->GetWorldFromContextObject(WorldContextObject);
 	check(World != nullptr);
 
-	ULevelSequencePlayer* NewPlayer = NewObject<ULevelSequencePlayer>(GetTransientPackage(), NAME_None, RF_Transient);
-	check(NewPlayer != nullptr);
-
-	NewPlayer->Initialize(InLevelSequence, World, Settings);
-
-	// Automatically tick this player
-	GAnimationPlayerTicker.Add(NewPlayer);
-
-	return NewPlayer;
-}
-
-
-bool ULevelSequencePlayer::IsPlaying() const
-{
-	return bIsPlaying;
-}
-
-
-void ULevelSequencePlayer::Pause()
-{
-	if (bIsPlaying)
-	{
-		if (bIsEvaluating)
-		{
-			LatentActions.Add(ELatentAction::Pause);
-			return;
-		}
-
-		bIsPlaying = false;
-
-		// Evaluate the sequence at its current time, with a status of 'stopped' to ensure that animated state pauses correctly
-		{
-			TGuardValue<bool> Guard(bIsEvaluating, true);
-
-			UMovieSceneSequence* MovieSceneSequence = RootTemplateInstance.GetSequence(MovieSceneSequenceID::Root);
-			TOptional<float> FixedFrameInterval = MovieSceneSequence->GetMovieScene() ? MovieSceneSequence->GetMovieScene()->GetOptionalFixedFrameInterval() : TOptional<float>();
-
-			FMovieSceneEvaluationRange Range = PlayPosition.JumpTo(GetSequencePosition(), FixedFrameInterval);
-
-			const FMovieSceneContext Context(Range, EMovieScenePlayerStatus::Stopped);
-			RootTemplateInstance.Evaluate(Context, *this);
-		}
-
-		ApplyLatentActions();
-
-		if (OnPause.IsBound())
-		{
-			OnPause.Broadcast();
-		}
-	}
-}
-
-void ULevelSequencePlayer::Stop()
-{
-	if (bIsPlaying)
-	{
-		if (bIsEvaluating)
-		{
-			LatentActions.Add(ELatentAction::Stop);
-			return;
-		}
+	FActorSpawnParameters SpawnParams;
+	SpawnParams.SpawnCollisionHandlingOverride = ESpawnActorCollisionHandlingMethod::AlwaysSpawn;
+	SpawnParams.ObjectFlags |= RF_Transient;
+	SpawnParams.bAllowDuringConstructionScript = true;
+	ALevelSequenceActor* Actor = World->SpawnActor<ALevelSequenceActor>(SpawnParams);
 
-		bIsPlaying = false;
-		TimeCursorPosition = PlaybackSettings.PlayRate < 0.f ? GetLength() : 0.f;
-		CurrentNumLoops = 0;
+	Actor->PlaybackSettings = Settings;
+	Actor->LevelSequence = InLevelSequence;
 
-		RootTemplateInstance.Finish(*this);
+	Actor->InitializePlayer();
 
-		SetTickPrerequisites(false);
-
-		if (OnStop.IsBound())
-		{
-			OnStop.Broadcast();
-		}
-	}
+	return Actor->SequencePlayer;
 }
 
 void GetDescendantSequences(UMovieSceneSequence* InSequence, TArray<UMovieSceneSequence*> & InSequences)
@@ -291,196 +169,33 @@ void ULevelSequencePlayer::SetTickPrerequisites(bool bAddTickPrerequisites)
 	}
 }
 
-void ULevelSequencePlayer::Play()
-{
-	bReversePlayback = false;
-
-	PlayInternal();
-}
-
-void ULevelSequencePlayer::PlayReverse() 
-{
-	bReversePlayback = true;
-
-	PlayInternal();
-}
-
-void ULevelSequencePlayer::ChangePlaybackDirection()
-{
-	bReversePlayback = !bReversePlayback;
-
-	PlayInternal();
-}
+/* ULevelSequencePlayer implementation
+ *****************************************************************************/
 
-void ULevelSequencePlayer::PlayInternal()
+void ULevelSequencePlayer::Initialize(ULevelSequence* InLevelSequence, UWorld* InWorld, const FMovieSceneSequencePlaybackSettings& Settings)
 {
-	if (!bIsPlaying)
-	{
-		// Start playing
-		StartPlayingNextTick();
-
-		// Update now
-		bPendingFirstUpdate = false;
-
-		UMovieSceneSequence* MovieSceneSequence = RootTemplateInstance.GetSequence(MovieSceneSequenceID::Root);
-		TOptional<float> FixedFrameInterval = MovieSceneSequence->GetMovieScene() ? MovieSceneSequence->GetMovieScene()->GetOptionalFixedFrameInterval() : TOptional<float>();
-
-		FMovieSceneEvaluationRange Range = PlayPosition.JumpTo(GetSequencePosition(), FixedFrameInterval);
-		UpdateMovieSceneInstance(Range);
-
-		if (OnPlay.IsBound())
-		{
-			OnPlay.Broadcast();
-		}
-	}
+	World = InWorld;
+	UMovieSceneSequencePlayer::Initialize(InLevelSequence, Settings);
 }
 
-void ULevelSequencePlayer::PlayLooping(int32 NumLoops)
+bool ULevelSequencePlayer::CanPlay() const
 {
-	PlaybackSettings.LoopCount = NumLoops;
-	PlayInternal();
+	return World.IsValid();
 }
 
-void ULevelSequencePlayer::StartPlayingNextTick()
+void ULevelSequencePlayer::OnStartedPlaying()
 {
-	if ((LevelSequence == nullptr) || !World.IsValid())
-	{
-		return;
-	}
-
-	// @todo Sequencer playback: Should we recreate the instance every time?
-	// We must not recreate the instance since it holds stateful information (such as which objects it has spawned). Recreating the instance would break any 
-	// @todo: Is this still the case now that eval state is stored (correctly) in the player?
-	if (!RootTemplateInstance.IsValid())
-	{
-		RootTemplateInstance.Initialize(*LevelSequence, *this);
-	}
-
 	SetTickPrerequisites(true);
-
-	bPendingFirstUpdate = true;
-	bIsPlaying = true;
-	bHasCleanedUpSequence = false;
 }
 
-float ULevelSequencePlayer::GetPlaybackPosition() const
+void ULevelSequencePlayer::OnStopped()
 {
-	return TimeCursorPosition;
+	SetTickPrerequisites(false);
 }
 
-void ULevelSequencePlayer::SetPlaybackPosition(float NewPlaybackPosition)
-{
-	UpdateTimeCursorPosition(NewPlaybackPosition);
-}
-
-float ULevelSequencePlayer::GetLength() const
-{
-	return EndTime - StartTime;
-}
-
-float ULevelSequencePlayer::GetPlayRate() const
-{
-	return PlaybackSettings.PlayRate;
-}
-
-void ULevelSequencePlayer::SetPlayRate(float PlayRate)
-{
-	PlaybackSettings.PlayRate = PlayRate;
-}
-
-void ULevelSequencePlayer::SetPlaybackRange( const float NewStartTime, const float NewEndTime )
-{
-	StartTime = NewStartTime;
-	EndTime = FMath::Max(NewEndTime, StartTime);
-
-	TimeCursorPosition = FMath::Clamp(TimeCursorPosition, 0.f, GetLength());
-}
-
-void ULevelSequencePlayer::UpdateTimeCursorPosition(float NewPosition)
-{
-	float Length = GetLength();
-
-	UMovieSceneSequence* MovieSceneSequence = RootTemplateInstance.GetSequence(MovieSceneSequenceID::Root);
-	TOptional<float> FixedFrameInterval = MovieSceneSequence->GetMovieScene() ? MovieSceneSequence->GetMovieScene()->GetOptionalFixedFrameInterval() : TOptional<float>();
-
-	if (bPendingFirstUpdate)
-	{
-		NewPosition = TimeCursorPosition;
-		bPendingFirstUpdate = false;
-	}
-
-	if (ShouldStopOrLoop(NewPosition))
-	{
-		// loop playback
-		if (PlaybackSettings.LoopCount < 0 || CurrentNumLoops < PlaybackSettings.LoopCount)
-		{
-			++CurrentNumLoops;
-			const float Overplay = FMath::Fmod(NewPosition, Length);
-			TimeCursorPosition = Overplay < 0 ? Length + Overplay : Overplay;
-			
-			PlayPosition.Reset(Overplay < 0 ? Length : 0.f);
-			FMovieSceneEvaluationRange Range = PlayPosition.PlayTo(GetSequencePosition(), FixedFrameInterval);
-
-			SpawnRegister->ForgetExternallyOwnedSpawnedObjects(State, *this);
-
-			UpdateMovieSceneInstance(Range);
-			return;
-		}
-
-		// stop playback
-		Stop();
-		return;
-	}
-
-	TimeCursorPosition = NewPosition;
-
-	FMovieSceneEvaluationRange Range = PlayPosition.PlayTo(NewPosition + StartTime, FixedFrameInterval);
-	UpdateMovieSceneInstance(Range);
-}
-
-bool ULevelSequencePlayer::ShouldStopOrLoop(float NewPosition)
-{
-	bool bShouldStopOrLoop = false;
-	if (!bReversePlayback)
-	{
-		bShouldStopOrLoop = NewPosition >= GetLength();
-	}
-	else
-	{
-		bShouldStopOrLoop = NewPosition < 0.f;
-	}
-
-	return bShouldStopOrLoop;
-}
-
-/* ULevelSequencePlayer implementation
- *****************************************************************************/
-
-void ULevelSequencePlayer::Initialize(ULevelSequence* InLevelSequence, UWorld* InWorld, const FLevelSequencePlaybackSettings& Settings)
-{
-	CurrentPlayer = this;
-	LevelSequence = InLevelSequence;
-
-	RootTemplateInstance.Initialize(*LevelSequence, *this);
-
-	World = InWorld;
-	PlaybackSettings = Settings;
-
-	if (UMovieScene* MovieScene = LevelSequence->GetMovieScene())
-	{
-		TRange<float> PlaybackRange = MovieScene->GetPlaybackRange();
-		SetPlaybackRange(PlaybackRange.GetLowerBoundValue(), PlaybackRange.GetUpperBoundValue());
-	}
-
-	// Ensure everything is set up, ready for playback
-	Stop();
-}
-
-
 /* IMovieScenePlayer interface
  *****************************************************************************/
 
-
 void ULevelSequencePlayer::UpdateCameraCut(UObject* CameraObject, UObject* UnlockIfCameraObject, bool bJumpCut)
 {
 	// skip missing player controller
@@ -554,23 +269,6 @@ void ULevelSequencePlayer::UpdateCameraCut(UObject* CameraObject, UObject* Unloc
 	}
 }
 
-void ULevelSequencePlayer::SetViewportSettings(const TMap<FViewportClient*, EMovieSceneViewportParams>& ViewportParamsMap)
-{
-}
-
-void ULevelSequencePlayer::GetViewportSettings(TMap<FViewportClient*, EMovieSceneViewportParams>& ViewportParamsMap) const
-{
-}
-
-EMovieScenePlayerStatus::Type ULevelSequencePlayer::GetPlaybackStatus() const
-{
-	return bIsPlaying ? EMovieScenePlayerStatus::Playing : EMovieScenePlayerStatus::Stopped;
-}
-
-void ULevelSequencePlayer::SetPlaybackStatus(EMovieScenePlayerStatus::Type InPlaybackStatus)
-{
-}
-
 UObject* ULevelSequencePlayer::GetPlaybackContext() const
 {
 	return World.Get();
@@ -602,64 +300,9 @@ void ULevelSequencePlayer::GetEventContexts(UWorld& InWorld, TArray<UObject*>& O
 	}
 }
 
-void ULevelSequencePlayer::Update(const float DeltaSeconds)
-{
-	if (bIsPlaying)
-	{
-		float PlayRate = bReversePlayback ? -PlaybackSettings.PlayRate : PlaybackSettings.PlayRate;
-		UpdateTimeCursorPosition(TimeCursorPosition + DeltaSeconds * PlayRate);
-	}
-	else if (!bHasCleanedUpSequence && ShouldStopOrLoop(TimeCursorPosition))
-	{
-		// Update the sequence one last time
-		UMovieSceneSequence* MovieSceneSequence = RootTemplateInstance.GetSequence(MovieSceneSequenceID::Root);
-		TOptional<float> FixedFrameInterval = MovieSceneSequence->GetMovieScene() ? MovieSceneSequence->GetMovieScene()->GetOptionalFixedFrameInterval() : TOptional<float>();
-		
-		FMovieSceneEvaluationRange Range = PlayPosition.PlayTo(GetSequencePosition(), FixedFrameInterval);
-		UpdateMovieSceneInstance(Range);
-
-		bHasCleanedUpSequence = true;
-		
-		SpawnRegister->ForgetExternallyOwnedSpawnedObjects(State, *this);
-		SpawnRegister->CleanUp(*this);
-	}
-}
-
-void ULevelSequencePlayer::UpdateMovieSceneInstance(FMovieSceneEvaluationRange InRange)
-{
-	{
-		TGuardValue<bool> Guard(bIsEvaluating, true);
-
-		const FMovieSceneContext Context(InRange, GetPlaybackStatus());
-		RootTemplateInstance.Evaluate(Context, *this);
-
-	#if WITH_EDITOR
-		OnLevelSequencePlayerUpdate.Broadcast(*this, Context.GetTime(), Context.GetPreviousTime());
-	#endif
-	}
-
-	ApplyLatentActions();
-}
-
-void ULevelSequencePlayer::ApplyLatentActions()
-{
-	// Swap to a stack array to ensure no reentrancy if we evaluate during a pause, for instance
-	TArray<ELatentAction> TheseActions;
-	Swap(TheseActions, LatentActions);
-
-	for (ELatentAction LatentAction : TheseActions)
-	{
-		switch(LatentAction)
-		{
-		case ELatentAction::Stop:	Stop(); break;
-		case ELatentAction::Pause:	Pause(); break;
-		}
-	}
-}
-
 void ULevelSequencePlayer::TakeFrameSnapshot(FLevelSequencePlayerSnapshot& OutSnapshot) const
 {
-	if (!ensure(LevelSequence))
+	if (!ensure(Sequence))
 	{
 		return;
 	}
@@ -669,13 +312,13 @@ void ULevelSequencePlayer::TakeFrameSnapshot(FLevelSequencePlayerSnapshot& OutSn
 	OutSnapshot.Settings = SnapshotSettings;
 
 	OutSnapshot.MasterTime = CurrentTime;
-	OutSnapshot.MasterName = FText::FromString(LevelSequence->GetName());
+	OutSnapshot.MasterName = FText::FromString(Sequence->GetName());
 
 	OutSnapshot.CurrentShotName = OutSnapshot.MasterName;
 	OutSnapshot.CurrentShotLocalTime = CurrentTime;
 	OutSnapshot.CameraComponent = CachedCameraComponent.IsValid() ? CachedCameraComponent.Get() : nullptr;
 
-	UMovieSceneCinematicShotTrack* ShotTrack = LevelSequence->GetMovieScene()->FindMasterTrack<UMovieSceneCinematicShotTrack>();
+	UMovieSceneCinematicShotTrack* ShotTrack = Sequence->GetMovieScene()->FindMasterTrack<UMovieSceneCinematicShotTrack>();
 	if (ShotTrack)
 	{
 		int32 HighestRow = TNumericLimits<int32>::Max();
@@ -708,7 +351,3 @@ void ULevelSequencePlayer::TakeFrameSnapshot(FLevelSequencePlayerSnapshot& OutSn
 	}
 }
 
-FMovieSceneSpawnRegister& ULevelSequencePlayer::GetSpawnRegister()
-{
-	return *SpawnRegister;
-}
diff --git a/Engine/Source/Runtime/LevelSequence/Private/LevelSequenceSpawnRegister.cpp b/Engine/Source/Runtime/LevelSequence/Private/LevelSequenceSpawnRegister.cpp
--- a/Engine/Source/Runtime/LevelSequence/Private/LevelSequenceSpawnRegister.cpp
+++ b/Engine/Source/Runtime/LevelSequence/Private/LevelSequenceSpawnRegister.cpp
@@ -2,6 +2,7 @@
 
 #include "LevelSequenceSpawnRegister.h"
 #include "Engine/EngineTypes.h"
+#include "IMovieScenePlayer.h"
 #include "MovieScene.h"
 #include "GameFramework/Actor.h"
 #include "Engine/World.h"
@@ -93,7 +94,8 @@ UObject* FLevelSequenceSpawnRegister::SpawnObject(FMovieSceneSpawnable& Spawnabl
 	// tag this actor so we know it was spawned by sequencer
 	SpawnedActor->Tags.Add(SequencerActorTag);
 
-	SpawnedActor->FinishSpawning(SpawnTransform);
+	const bool bIsDefaultTransform = true;
+	SpawnedActor->FinishSpawning(SpawnTransform, bIsDefaultTransform);
 
 	return SpawnedActor;
 }
diff --git a/Engine/Source/Runtime/LevelSequence/Public/LevelSequence.h b/Engine/Source/Runtime/LevelSequence/Public/LevelSequence.h
--- a/Engine/Source/Runtime/LevelSequence/Public/LevelSequence.h
+++ b/Engine/Source/Runtime/LevelSequence/Public/LevelSequence.h
@@ -44,8 +44,8 @@ public:
 	virtual UObject* GetParentObject(UObject* Object) const override;
 	virtual void UnbindPossessableObjects(const FGuid& ObjectId) override;
 	virtual bool AllowsSpawnableObjects() const override;
+	virtual bool CanRebindPossessable(const FMovieScenePossessable& InPossessable) const override;
 	virtual UObject* MakeSpawnableTemplateFromInstance(UObject& InSourceObject, FName ObjectName) override;
-	
 
 	virtual void PostLoad() override;
 
diff --git a/Engine/Source/Runtime/LevelSequence/Public/LevelSequenceActor.h b/Engine/Source/Runtime/LevelSequence/Public/LevelSequenceActor.h
--- a/Engine/Source/Runtime/LevelSequence/Public/LevelSequenceActor.h
+++ b/Engine/Source/Runtime/LevelSequence/Public/LevelSequenceActor.h
@@ -9,17 +9,19 @@
 #include "GameFramework/Actor.h"
 #include "Misc/StringClassReference.h"
 #include "LevelSequencePlayer.h"
+#include "MovieSceneBindingOwnerInterface.h"
+#include "MovieSceneBindingOverrides.h"
 #include "LevelSequenceActor.generated.h"
 
 class ULevelSequenceBurnIn;
 
-UCLASS(Blueprintable)
+UCLASS(Blueprintable, DefaultToInstanced)
 class LEVELSEQUENCE_API ULevelSequenceBurnInInitSettings : public UObject
 {
 	GENERATED_BODY()
 };
 
-UCLASS()
+UCLASS(DefaultToInstanced)
 class LEVELSEQUENCE_API ULevelSequenceBurnInOptions : public UObject
 {
 public:
@@ -38,7 +40,7 @@ public:
 	UPROPERTY(EditAnywhere, Category="General", meta=(EditCondition=bUseBurnIn, MetaClass="LevelSequenceBurnIn"))
 	FStringClassReference BurnInClass;
 
-	UPROPERTY(EditAnywhere, Category="General", meta=(EditCondition=bUseBurnIn, EditInline))
+	UPROPERTY(Instanced, EditAnywhere, Category="General", meta=(EditCondition=bUseBurnIn))
 	ULevelSequenceBurnInInitSettings* Settings;
 
 protected:
@@ -54,6 +56,7 @@ protected:
 UCLASS(hideCategories=(Rendering, Physics, LOD, Activation))
 class LEVELSEQUENCE_API ALevelSequenceActor
 	: public AActor
+	, public IMovieSceneBindingOwnerInterface
 {
 public:
 
@@ -68,17 +71,21 @@ public:
 	bool bAutoPlay;
 
 	UPROPERTY(EditAnywhere, BlueprintReadOnly, Category="Playback", meta=(ShowOnlyInnerProperties))
-	FLevelSequencePlaybackSettings PlaybackSettings;
+	FMovieSceneSequencePlaybackSettings PlaybackSettings;
 
-	UPROPERTY(transient, BlueprintReadOnly, Category="Playback")
+	UPROPERTY(transient, BlueprintReadOnly, Category="Playback", meta=(ExposeFunctionCategories="Game|Cinematic"))
 	ULevelSequencePlayer* SequencePlayer;
 
 	UPROPERTY(EditAnywhere, BlueprintReadOnly, Category="General", meta=(AllowedClasses="LevelSequence"))
 	FStringAssetReference LevelSequence;
 
-	UPROPERTY(EditAnywhere, BlueprintReadWrite, AdvancedDisplay, Category="General", meta=(EditInline))
+	UPROPERTY(Instanced, VisibleAnywhere, AdvancedDisplay, BlueprintReadOnly, Category="General")
 	ULevelSequenceBurnInOptions* BurnInOptions;
 
+	/** Mapping of actors to override the sequence bindings with */
+	UPROPERTY(Instanced, VisibleAnywhere, AdvancedDisplay, BlueprintReadOnly, Category="General")
+	UMovieSceneBindingOverrides* BindingOverrides;
+
 public:
 
 	/**
@@ -103,6 +110,63 @@ public:
 	/** Refresh this actor's burn in */
 	void RefreshBurnIn();
 
+public:
+
+	/** Overrides the specified binding with the specified actors, optionally still allowing the bindings defined in the Level Sequence asset */
+	UFUNCTION(BlueprintCallable, Category="Game|Cinematic|Bindings")
+	void SetBinding(FMovieSceneObjectBindingPtr Binding, const TArray<AActor*>& Actors, bool bAllowBindingsFromAsset = false)
+	{
+		BindingOverrides->SetBinding(Binding, TArray<UObject*>(Actors), bAllowBindingsFromAsset);
+		if (SequencePlayer)
+		{
+			SequencePlayer->State.Invalidate(Binding.Guid, MovieSceneSequenceID::Root);
+		}
+	}
+
+	/** Adds the specified actor to the overridden bindings for the specified binding ID, optionally still allowing the bindings defined in the Level Sequence asset */
+	UFUNCTION(BlueprintCallable, Category="Game|Cinematic|Bindings")
+	void AddBinding(FMovieSceneObjectBindingPtr Binding, AActor* Actor, bool bAllowBindingsFromAsset = false)
+	{
+		BindingOverrides->AddBinding(Binding, Actor);
+		if (SequencePlayer)
+		{
+			SequencePlayer->State.Invalidate(Binding.Guid, MovieSceneSequenceID::Root);
+		}
+	}
+
+	/** Removes the specified actor from the specified binding's actor array */
+	UFUNCTION(BlueprintCallable, Category="Game|Cinematic|Bindings")
+	void RemoveBinding(FMovieSceneObjectBindingPtr Binding, AActor* Actor)
+	{
+		BindingOverrides->RemoveBinding(Binding, Actor);
+		if (SequencePlayer)
+		{
+			SequencePlayer->State.Invalidate(Binding.Guid, MovieSceneSequenceID::Root);
+		}
+	}
+
+	/** Resets the specified binding back to the defaults defined by the Level Sequence asset */
+	UFUNCTION(BlueprintCallable, Category="Game|Cinematic|Bindings")
+	void ResetBinding(FMovieSceneObjectBindingPtr Binding)
+	{
+		BindingOverrides->ResetBinding(Binding);
+		if (SequencePlayer)
+		{
+			SequencePlayer->State.Invalidate(Binding.Guid, MovieSceneSequenceID::Root);
+		}
+	}
+
+	/** Resets all overridden bindings back to the defaults defined by the Level Sequence asset */
+	UFUNCTION(BlueprintCallable, Category="Game|Cinematic|Bindings")
+	void ResetBindings()
+	{
+		BindingOverrides->ResetBindings();
+		if (SequencePlayer)
+		{
+			SequencePlayer->State.ClearObjectCaches();
+		}
+	}
+
 public:
 
 	virtual void PostInitializeComponents() override;
@@ -120,8 +184,38 @@ public:
 
 	void InitializePlayer();
 
+#if WITH_EDITOR
+	virtual TSharedPtr<FStructOnScope> GetObjectPickerProxy(TSharedPtr<IPropertyHandle> PropertyHandle) override;
+	virtual void UpdateObjectFromProxy(FStructOnScope& Proxy, IPropertyHandle& ObjectPropertyHandle) override;
+	virtual UMovieSceneSequence* RetrieveOwnedSequence() const override
+	{
+		return GetSequence(true);
+	}
+#endif
+
 private:
 	/** Burn-in widget */
 	UPROPERTY()
 	ULevelSequenceBurnIn* BurnInInstance;
 };
+
+
+USTRUCT()
+struct FBoundActorProxy
+{
+	GENERATED_BODY()
+
+#if WITH_EDITORONLY_DATA
+
+	/** Specifies the actor to override the binding with */
+	UPROPERTY(EditInstanceOnly, AdvancedDisplay, Category="General")
+	AActor* BoundActor;
+
+	void Initialize(TSharedPtr<IPropertyHandle> InPropertyHandle);
+
+	void OnReflectedPropertyChanged();
+
+	TSharedPtr<IPropertyHandle> ReflectedProperty;
+
+#endif
+};
\ No newline at end of file
diff --git a/Engine/Source/Runtime/LevelSequence/Public/LevelSequenceBurnIn.h b/Engine/Source/Runtime/LevelSequence/Public/LevelSequenceBurnIn.h
--- a/Engine/Source/Runtime/LevelSequence/Public/LevelSequenceBurnIn.h
+++ b/Engine/Source/Runtime/LevelSequence/Public/LevelSequenceBurnIn.h
@@ -36,7 +36,7 @@ public:
 protected:
 
 	/** Called as part of the game tick loop when the sequence has been updated */
-	void OnSequenceUpdated(const ULevelSequencePlayer& Player, float CurrentTime, float PreviousTime);
+	void OnSequenceUpdated(const UMovieSceneSequencePlayer& Player, float CurrentTime, float PreviousTime);
 
 protected:
 
diff --git a/Engine/Source/Runtime/LevelSequence/Public/LevelSequencePlayer.h b/Engine/Source/Runtime/LevelSequence/Public/LevelSequencePlayer.h
--- a/Engine/Source/Runtime/LevelSequence/Public/LevelSequencePlayer.h
+++ b/Engine/Source/Runtime/LevelSequence/Public/LevelSequencePlayer.h
@@ -9,6 +9,7 @@
 #include "IMovieScenePlayer.h"
 #include "Evaluation/MovieScenePlayback.h"
 #include "Evaluation/MovieSceneEvaluationTemplateInstance.h"
+#include "MovieSceneSequencePlayer.h"
 #include "LevelSequence.h"
 #include "LevelSequencePlayer.generated.h"
 
@@ -17,7 +18,9 @@ class FLevelSequenceSpawnRegister;
 class FViewportClient;
 class UCameraComponent;
 
-DECLARE_DYNAMIC_MULTICAST_DELEGATE(FOnLevelSequencePlayerEvent);
+struct DEPRECATED(4.15, "Please use FMovieSceneSequencePlaybackSettings.") FLevelSequencePlaybackSettings
+	: public FMovieSceneSequencePlaybackSettings
+{};
 
 USTRUCT(BlueprintType)
 struct FLevelSequenceSnapshotSettings
@@ -68,28 +71,6 @@ struct FLevelSequencePlayerSnapshot
 	FLevelSequenceSnapshotSettings Settings;
 };
 
-/**
- * Settings for the level sequence player actor.
- */
-USTRUCT(BlueprintType)
-struct FLevelSequencePlaybackSettings
-{
-	FLevelSequencePlaybackSettings()
-		: LoopCount(0)
-		, PlayRate(1.f)
-	{ }
-
-	GENERATED_BODY()
-
-	/** Number of times to loop playback. -1 for infinite, else the number of times to loop before stopping */
-	UPROPERTY(EditAnywhere, BlueprintReadWrite, Category="Playback", meta=(UIMin=1, DisplayName="Loop"))
-	int32 LoopCount;
-
-	/** The rate at which to playback the animation */
-	UPROPERTY(EditAnywhere, BlueprintReadWrite, Category="Playback", meta=(Units=Multiplier))
-	float PlayRate;
-};
-
 /**
  * ULevelSequencePlayer is used to actually "play" an level sequence asset at runtime.
  *
@@ -98,8 +79,7 @@ struct FLevelSequencePlaybackSettings
  */
 UCLASS(BlueprintType)
 class LEVELSEQUENCE_API ULevelSequencePlayer
-	: public UObject
-	, public IMovieScenePlayer
+	: public UMovieSceneSequencePlayer
 {
 public:
 	ULevelSequencePlayer(const FObjectInitializer&);
@@ -113,7 +93,7 @@ public:
 	 * @param InWorld The world that the animation is played in.
 	 * @param Settings The desired playback settings
 	 */
-	void Initialize(ULevelSequence* InLevelSequence, UWorld* InWorld, const FLevelSequencePlaybackSettings& Settings);
+	void Initialize(ULevelSequence* InLevelSequence, UWorld* InWorld, const FMovieSceneSequencePlaybackSettings& Settings);
 
 public:
 
@@ -125,86 +105,7 @@ public:
 	 * @param Settings The desired playback settings
 	 */
 	UFUNCTION(BlueprintCallable, Category="Game|Cinematic", meta=(WorldContext="WorldContextObject"))
-	static ULevelSequencePlayer* CreateLevelSequencePlayer(UObject* WorldContextObject, ULevelSequence* LevelSequence, FLevelSequencePlaybackSettings Settings);
-
-	/** Start playback forwards from the current time cursor position, using the current play rate. */
-	UFUNCTION(BlueprintCallable, Category="Game|Cinematic")
-	void Play();
-
-	/** Reverse playback. */
-	UFUNCTION(BlueprintCallable, Category="Game|Cinematic")
-	void PlayReverse();
-
-	/** Changes the direction of playback (go in reverse if it was going forward, or vice versa) */
-	UFUNCTION(BlueprintCallable, Category="Game|Cinematic")
-	void ChangePlaybackDirection();
-
-	/**
-	 * Start playback from the current time cursor position, looping the specified number of times.
-	 * @param NumLoops - The number of loops to play. -1 indicates infinite looping.
-	 */
-	UFUNCTION(BlueprintCallable, Category="Game|Cinematic")
-	void PlayLooping(int32 NumLoops = -1);
-
-	/** Start playback from the current time cursor position, using the current play rate. Does not update the animation until next tick. */
-	UFUNCTION(BlueprintCallable, Category="Game|Cinematic")
-	void StartPlayingNextTick();
-	
-	/** Pause playback. */
-	UFUNCTION(BlueprintCallable, Category="Game|Cinematic")
-	void Pause();
-	
-	/** Stop playback. */
-	UFUNCTION(BlueprintCallable, Category="Game|Cinematic")
-	void Stop();
-
-	/** Get the current playback position */
-	UFUNCTION(BlueprintCallable, Category="Game|Cinematic")
-	float GetPlaybackPosition() const;
-
-	/**
-	 * Set the current playback position
-	 * @param NewPlaybackPosition - The new playback position to set.
-	 * If the animation is currently playing, it will continue to do so from the new position
-	 */
-	UFUNCTION(BlueprintCallable, Category="Game|Cinematic")
-	void SetPlaybackPosition(float NewPlaybackPosition);
-
-	/** Check whether the sequence is actively playing. */
-	UFUNCTION(BlueprintCallable, Category="Game|Cinematic")
-	bool IsPlaying() const;
-
-	/** Get the playback length of the sequence */
-	UFUNCTION(BlueprintCallable, Category="Game|Cinematic")
-	float GetLength() const;
-
-	/** Get the playback rate of this player. */
-	UFUNCTION(BlueprintCallable, Category="Game|Cinematic")
-	float GetPlayRate() const;
-
-	/**
-	 * Set the playback rate of this player. Negative values will play the animation in reverse.
-	 * @param PlayRate - The new rate of playback for the animation.
-	 */
-	UFUNCTION(BlueprintCallable, Category="Game|Cinematic")
-	void SetPlayRate(float PlayRate);
-
-	/**
-	 * Sets the range in time to be played back by this player, overriding the default range stored in the asset
-	 *
-	 * @param	NewStartTime	The new starting time for playback
-	 * @param	NewEndTime		The new ending time for playback.  Must be larger than the start time.
-	 */
-	UFUNCTION(BlueprintCallable, Category="Game|Cinematic")
-	void SetPlaybackRange( const float NewStartTime, const float NewEndTime );
-
-	/** Get the offset within the level sequence to start playing */
-	UFUNCTION(BlueprintCallable, Category="Game|Cinematic")
-	float GetPlaybackStart() const { return StartTime; }
-
-	/** Get the offset within the level sequence to finish playing */
-	UFUNCTION(BlueprintCallable, Category="Game|Cinematic")
-	float GetPlaybackEnd() const { return EndTime; }
+	static ULevelSequencePlayer* CreateLevelSequencePlayer(UObject* WorldContextObject, ULevelSequence* LevelSequence, FMovieSceneSequencePlaybackSettings Settings);
 
 	/** Set the settings used to capture snapshots with */
 	void SetSnapshotSettings(const FLevelSequenceSnapshotSettings& InSettings) { SnapshotSettings = InSettings; }
@@ -215,147 +116,46 @@ public:
 	 * Access the level sequence this player is playing
 	 * @return the level sequence currently assigned to this player
 	 */
-	ULevelSequence* GetLevelSequence() const { return LevelSequence; }
-	
+	DEPRECATED(4.15, "Please use GetSequence instead.")
+	ULevelSequence* GetLevelSequence() const { return Cast<ULevelSequence>(Sequence); }
+
 protected:
 
 	// IMovieScenePlayer interface
-	virtual FMovieSceneRootEvaluationTemplateInstance& GetEvaluationTemplate() override { return RootTemplateInstance; }
 	virtual void UpdateCameraCut(UObject* CameraObject, UObject* UnlockIfCameraObject, bool bJumpCut) override;
-	virtual void SetViewportSettings(const TMap<FViewportClient*, EMovieSceneViewportParams>& ViewportParamsMap) override;
-	virtual void GetViewportSettings(TMap<FViewportClient*, EMovieSceneViewportParams>& ViewportParamsMap) const override;
-	virtual EMovieScenePlayerStatus::Type GetPlaybackStatus() const override;
-	virtual void SetPlaybackStatus(EMovieScenePlayerStatus::Type InPlaybackStatus) override;
-	virtual FMovieSceneSpawnRegister& GetSpawnRegister() override;
 	virtual UObject* GetPlaybackContext() const override;
 	virtual TArray<UObject*> GetEventContexts() const override;
 
+	//~ UMovieSceneSequencePlayer interface
+	virtual bool CanPlay() const override;
+	virtual void OnStartedPlaying() override;
+	virtual void OnStopped() override;
+
 public:
 
 	/** Populate the specified array with any given event contexts for the specified world */
 	static void GetEventContexts(UWorld& InWorld, TArray<UObject*>& OutContexts);
 
-	void Update(const float DeltaSeconds);
-
 	/** Take a snapshot of the current state of this player */
 	void TakeFrameSnapshot(FLevelSequencePlayerSnapshot& OutSnapshot) const;
 
 private:
 
-	/** Internal play function. */
-	void PlayInternal();
-
-	/** Update the movie scene instance from the specified previous position, to the specified time position. */
-	void UpdateMovieSceneInstance(FMovieSceneEvaluationRange InRange);
-
-	/** Update the time cursor position and handle stopping & looping. */
-	void UpdateTimeCursorPosition(float NewPosition);
-
 	/** Add tick prerequisites so that the level sequence actor ticks before all the actors it controls */
 	void SetTickPrerequisites(bool bAddTickPrerequisites);
-	
-	/** Returns whether playback should be stopped or looped because it has gone out of the playback bounds. */
-	bool ShouldStopOrLoop(float NewPosition);
-
-	/** Get the position of the cursor in the root sequence */
-	FORCEINLINE float GetSequencePosition() const { return TimeCursorPosition + StartTime; }
-
-	/** Apply any latent actions which may have accumulated while the sequence was being evaluated */
-	void ApplyLatentActions();
-
-private:
-	enum class ELatentAction
-	{
-		Stop, Pause
-	};
-
-	/** The level sequence to play. */
-	UPROPERTY(transient)
-	ULevelSequence* LevelSequence;
-
-	/** The level sequence player. */
-	UPROPERTY(transient)
-	ULevelSequencePlayer* CurrentPlayer;
-
-	/** Whether we're currently playing. If false, then sequence playback is paused or was never started. */
-	UPROPERTY()
-	bool bIsPlaying;
-
-	/** Whether we're currently playing in reverse. */
-	UPROPERTY()
-	bool bReversePlayback;
-
-	/** True where we're waiting for the first update of the sequence after calling StartPlayingNextTick. */
-	bool bPendingFirstUpdate;
-
-	/** The current time cursor position within the sequence (in seconds) */
-	UPROPERTY()
-	float TimeCursorPosition;
-
-	/** Time time at which to start playing the sequence (defaults to the lower bound of the sequence's play range) */
-	float StartTime;
-
-	/** Time time at which to end playing the sequence (defaults to the upper bound of the sequence's play range) */
-	float EndTime;
-
-	/** Specific playback settings for the animation. */
-	UPROPERTY()
-	FLevelSequencePlaybackSettings PlaybackSettings;
-
-	/** The number of times we have looped in the current playback */
-	int32 CurrentNumLoops;
-
-	/** Whether this player has cleaned up the level sequence after it has stopped playing or not */
-	bool bHasCleanedUpSequence;
-
-	/** Set to true while evaluating to prevent reentrancy */
-	bool bIsEvaluating;
-
-	/** Set of latent actions that are to be performed when the sequence has finished evaluating this frame */
-	TArray<ELatentAction> LatentActions;
 
 private:
 
 	/** The world this player will spawn actors in, if needed */
 	TWeakObjectPtr<UWorld> World;
 
-	/** Register responsible for managing spawned objects */
-	TSharedPtr<FLevelSequenceSpawnRegister> SpawnRegister;
-
 	/** The last view target to reset to when updating camera cuts to null */
 	TWeakObjectPtr<AActor> LastViewTarget;
 
-	FMovieSceneRootEvaluationTemplateInstance RootTemplateInstance;
-
-	FMovieScenePlaybackPosition PlayPosition;
-
 protected:
 
 	/** How to take snapshots */
 	FLevelSequenceSnapshotSettings SnapshotSettings;
 
-public:
-
-	/** An event that is broadcast each time this level sequence player is updated */
-	DECLARE_EVENT_ThreeParams( ULevelSequencePlayer, FOnLevelSequencePlayerUpdated, const ULevelSequencePlayer&, float /*current time*/, float /*previous time*/ );
-	FOnLevelSequencePlayerUpdated& OnSequenceUpdated() const { return OnLevelSequencePlayerUpdate; }
-
-	/** Event triggered when the level sequence player is played */
-	UPROPERTY(BlueprintAssignable, Category="Game|Cinematic")
-	FOnLevelSequencePlayerEvent OnPlay;
-
-	/** Event triggered when the level sequence player is stopped */
-	UPROPERTY(BlueprintAssignable, Category="Game|Cinematic")
-	FOnLevelSequencePlayerEvent OnStop;
-
-	/** Event triggered when the level sequence player is paused */
-	UPROPERTY(BlueprintAssignable, Category="Game|Cinematic")
-	FOnLevelSequencePlayerEvent OnPause;
-
-private:
-
-	/** The event that will be broadcast every time the sequence is updated */
-	mutable FOnLevelSequencePlayerUpdated OnLevelSequencePlayerUpdate;
-
 	TWeakObjectPtr<UCameraComponent> CachedCameraComponent;
 };
diff --git a/Engine/Source/Runtime/Media/Public/IMediaBinarySink.h b/Engine/Source/Runtime/Media/Public/IMediaBinarySink.h
new file mode 100644
--- /dev/null
+++ b/Engine/Source/Runtime/Media/Public/IMediaBinarySink.h
@@ -0,0 +1,45 @@
+// Copyright 1998-2015 Epic Games, Inc. All Rights Reserved.
+
+#pragma once
+
+
+/**
+ * Interface for media sinks that receive binary data.
+ *
+ * Implementations of this interface must be thread-safe as these
+ * methods will be called from some random media decoder thread.
+ */
+class IMediaBinarySink
+{
+public:
+
+	/**
+	 * Initialize the sink.
+	 *
+	 * @return true on success, false otherwise.
+	 * @see ShutdownBinarySink
+	 */
+	virtual bool InitializeBinarySink() = 0;
+
+	/**
+	 * Process the given binary data.
+	 *
+	 * @param Data The data process.
+	 * @param Size The size of the data (in bytes).
+	 * @param Time The timestamp of the data (relative to the beginning of the media).
+	 * @param Duration The duration of the data.
+	 */
+	virtual void ProcessBinarySinkData(const uint8* Data, uint32 Size, FTimespan Time, FTimespan Duration) = 0;
+
+	/**
+	 * Shut down the sink.
+	 *
+	 * @see InitializeBinarySink
+	 */
+	virtual void ShutdownBinarySink() = 0;
+
+public:
+
+	/** Virtual destructor. */
+	virtual ~IMediaBinarySink() { }
+};
