Sequencer: Fixed spawnables not working in cooked builds

diff --git a/Engine/Source/Runtime/LevelSequence/Private/LevelSequenceSpawnRegister.cpp b/Engine/Source/Runtime/LevelSequence/Private/LevelSequenceSpawnRegister.cpp
--- a/Engine/Source/Runtime/LevelSequence/Private/LevelSequenceSpawnRegister.cpp
+++ b/Engine/Source/Runtime/LevelSequence/Private/LevelSequenceSpawnRegister.cpp
@@ -10,22 +10,20 @@
 
 FLevelSequenceSpawnRegister::FLevelSequenceSpawnRegister()
 {
-	FLevelSequenceModule* LevelSequenceModule = FModuleManager::GetModulePtr<FLevelSequenceModule>("LevelSequence");
-	if (LevelSequenceModule)
+	FLevelSequenceModule& LevelSequenceModule = FModuleManager::LoadModuleChecked<FLevelSequenceModule>("LevelSequence");
+	
+	for (const FOnCreateMovieSceneObjectSpawner& OnCreateMovieSceneObjectSpawner : LevelSequenceModule.OnCreateMovieSceneObjectSpawnerDelegates)
 	{
-		for (const FOnCreateMovieSceneObjectSpawner& OnCreateMovieSceneObjectSpawner : LevelSequenceModule->OnCreateMovieSceneObjectSpawnerDelegates)
-		{
-			check(OnCreateMovieSceneObjectSpawner.IsBound());
-			MovieSceneObjectSpawners.Add(OnCreateMovieSceneObjectSpawner.Execute());
-		}
-
-		// Now sort the spawners. Editor spawners should come first so they override runtime versions of the same supported type in-editor.
-		// @TODO: we could also sort by most-derived type here to allow for type specific behaviors
-		MovieSceneObjectSpawners.Sort([](TSharedRef<IMovieSceneObjectSpawner> LHS, TSharedRef<IMovieSceneObjectSpawner> RHS)
-		{
-			return LHS->IsEditor() > RHS->IsEditor();
-		});
+		check(OnCreateMovieSceneObjectSpawner.IsBound());
+		MovieSceneObjectSpawners.Add(OnCreateMovieSceneObjectSpawner.Execute());
 	}
+
+	// Now sort the spawners. Editor spawners should come first so they override runtime versions of the same supported type in-editor.
+	// @TODO: we could also sort by most-derived type here to allow for type specific behaviors
+	MovieSceneObjectSpawners.Sort([](TSharedRef<IMovieSceneObjectSpawner> LHS, TSharedRef<IMovieSceneObjectSpawner> RHS)
+	{
+		return LHS->IsEditor() > RHS->IsEditor();
+	});
 }
 
 UObject* FLevelSequenceSpawnRegister::SpawnObject(FMovieSceneSpawnable& Spawnable, FMovieSceneSequenceIDRef TemplateID, IMovieScenePlayer& Player)
