Sequencer: Allow resolve by path if fixup for PIE fails.

Copy from Release-4.16

diff --git a/Engine/Source/Runtime/LevelSequence/Private/LevelSequenceObjectReference.cpp b/Engine/Source/Runtime/LevelSequence/Private/LevelSequenceObjectReference.cpp
--- a/Engine/Source/Runtime/LevelSequence/Private/LevelSequenceObjectReference.cpp
+++ b/Engine/Source/Runtime/LevelSequence/Private/LevelSequenceObjectReference.cpp
@@ -46,6 +46,14 @@ UObject* FLevelSequenceObjectReference::Resolve(UObject* InContext) const
 
 		if (PIEInstanceID != -1 && FixedUpId == ObjectId)
 		{
+			if (!ObjectPath.IsEmpty())
+			{
+				if (UObject* FoundObject = FindObject<UObject>(InContext, *ObjectPath, false))
+				{
+					return FoundObject;
+				}
+			}
+
 			UE_LOG(LogMovieScene, Warning, TEXT("Attempted to resolve object with a PIE instance that has not been fixed up yet. This is probably due to a streamed level not being available yet."));
 			return nullptr;
 		}
@@ -65,11 +73,6 @@ UObject* FLevelSequenceObjectReference::Resolve(UObject* InContext) const
 		{
 			return FoundObject;
 		}
-
-		if (UObject* FoundObject = FindObject<UObject>(ANY_PACKAGE, *ObjectPath, false))
-		{
-			return FoundObject;
-		}
 	}
 
 	return nullptr;
