Merging CL 3421302 from Release-4.16

diff --git a/Engine/Source/Runtime/LevelSequence/Private/LevelSequenceObjectReference.cpp b/Engine/Source/Runtime/LevelSequence/Private/LevelSequenceObjectReference.cpp
--- a/Engine/Source/Runtime/LevelSequence/Private/LevelSequenceObjectReference.cpp
+++ b/Engine/Source/Runtime/LevelSequence/Private/LevelSequenceObjectReference.cpp
@@ -52,6 +52,11 @@ UObject* FLevelSequenceObjectReference::Resolve(UObject* InContext) const
 				{
 					return FoundObject;
 				}
+
+				if (UObject* FoundObject = FindObject<UObject>(ANY_PACKAGE, *ObjectPath, false))
+				{
+					return FoundObject;
+				}
 			}
 
 			UE_LOG(LogMovieScene, Warning, TEXT("Attempted to resolve object with a PIE instance that has not been fixed up yet. This is probably due to a streamed level not being available yet. %s"), *ObjectPath);
@@ -73,6 +78,11 @@ UObject* FLevelSequenceObjectReference::Resolve(UObject* InContext) const
 		{
 			return FoundObject;
 		}
+
+		if (UObject* FoundObject = FindObject<UObject>(ANY_PACKAGE, *ObjectPath, false))
+		{
+			return FoundObject;
+		}
 	}
 
 	return nullptr;
