Copying //UE4/Dev-Framework to //UE4/Dev-Main (Source: //UE4/Dev-Framework @ 3582324)

diff --git a/Engine/Source/Runtime/Landscape/Classes/LandscapeSplinesComponent.h b/Engine/Source/Runtime/Landscape/Classes/LandscapeSplinesComponent.h
--- a/Engine/Source/Runtime/Landscape/Classes/LandscapeSplinesComponent.h
+++ b/Engine/Source/Runtime/Landscape/Classes/LandscapeSplinesComponent.h
@@ -133,7 +133,7 @@ protected:
 #if WITH_EDITORONLY_DATA
 	// Serialized
 	UPROPERTY(TextExportTransient)
-	TMap<TAssetPtr<UWorld>, FForeignWorldSplineData> ForeignWorldSplineDataMap;
+	TMap<TSoftObjectPtr<UWorld>, FForeignWorldSplineData> ForeignWorldSplineDataMap;
 
 	// Transient - rebuilt on load
 	TMap<UMeshComponent*, UObject*> MeshComponentLocalOwnersMap;
diff --git a/Engine/Source/Runtime/Landscape/Private/LandscapeSplines.cpp b/Engine/Source/Runtime/Landscape/Private/LandscapeSplines.cpp
--- a/Engine/Source/Runtime/Landscape/Private/LandscapeSplines.cpp
+++ b/Engine/Source/Runtime/Landscape/Private/LandscapeSplines.cpp
@@ -562,7 +562,7 @@ void ULandscapeSplinesComponent::AutoFixMeshComponentErrors(UWorld* OtherWorld)
 {
 	UWorld* ThisOuterWorld = GetTypedOuter<UWorld>();
 
-	TAssetPtr<UWorld> OtherWorldAssetPtr = OtherWorld;
+	TSoftObjectPtr<UWorld> OtherWorldSoftPtr = OtherWorld;
 	ULandscapeSplinesComponent* StreamingSplinesComponent = GetStreamingSplinesComponentForLevel(OtherWorld->PersistentLevel);
 	auto* ForeignWorldSplineData = StreamingSplinesComponent ? StreamingSplinesComponent->ForeignWorldSplineDataMap.Find(ThisOuterWorld) : nullptr;
 
@@ -634,9 +634,9 @@ void ULandscapeSplinesComponent::CheckForErrors()
 	// Check spline segment meshes
 	for (ULandscapeSplineSegment* Segment : Segments)
 	{
-		for (auto& ForeignWorldAssetPtr : Segment->GetForeignWorlds())
+		for (auto& ForeignWorldSoftPtr : Segment->GetForeignWorlds())
 		{
-			UWorld* ForeignWorld = ForeignWorldAssetPtr.Get();
+			UWorld* ForeignWorld = ForeignWorldSoftPtr.Get();
 
 			if (ForeignWorld && !OutdatedWorlds.Contains(ForeignWorld))
 			{
@@ -674,16 +674,16 @@ void ULandscapeSplinesComponent::CheckForErrors()
 	// check for orphaned components
 	for (auto& ForeignWorldSplineDataPair : ForeignWorldSplineDataMap)
 	{
-		auto& ForeignWorldAssetPtr = ForeignWorldSplineDataPair.Key;
+		auto& ForeignWorldSoftPtr = ForeignWorldSplineDataPair.Key;
 		auto& ForeignWorldSplineData = ForeignWorldSplineDataPair.Value;
 
 		// World is not loaded
-		if (ForeignWorldAssetPtr.IsPending())
+		if (ForeignWorldSoftPtr.IsPending())
 		{
 			continue;
 		}
 
-		UWorld* ForeignWorld = ForeignWorldAssetPtr.Get();
+		UWorld* ForeignWorld = ForeignWorldSoftPtr.Get();
 		for (auto& ForeignSplineSegmentData : ForeignWorldSplineData.ForeignSplineSegmentData)
 		{
 			const ULandscapeSplineSegment* ForeignSplineSegment = ForeignSplineSegmentData.Identifier.Get();
diff --git a/Engine/Source/Runtime/LevelSequence/Private/LevelSequenceBindingReference.cpp b/Engine/Source/Runtime/LevelSequence/Private/LevelSequenceBindingReference.cpp
--- a/Engine/Source/Runtime/LevelSequence/Private/LevelSequenceBindingReference.cpp
+++ b/Engine/Source/Runtime/LevelSequence/Private/LevelSequenceBindingReference.cpp
@@ -24,59 +24,49 @@ FLevelSequenceBindingReference::FLevelSequenceBindingReference(UObject* InObject
 			return;
 		}
 
-		PackageName = ObjectPackage->GetName();
-	#if WITH_EDITORONLY_DATA
+		FString PackageName = ObjectPackage->GetName();
+#if WITH_EDITORONLY_DATA
+		// If this is being set from PIE we need to remove the pie prefix and point to the editor object
 		if (ObjectPackage->PIEInstanceID != INDEX_NONE)
 		{
 			FString PIEPrefix = FString::Printf(PLAYWORLD_PACKAGE_PREFIX TEXT("_%d_"), ObjectPackage->PIEInstanceID);
 			PackageName.ReplaceInline(*PIEPrefix, TEXT(""));
 		}
-	#endif
-
-		ObjectPath = InObject->GetPathName(ObjectPackage);
+#endif
+		
+		FString FullPath = PackageName + TEXT(".") + InObject->GetPathName(ObjectPackage);
+		ExternalObjectPath = FSoftObjectPath(FullPath);
 	}
 }
 
 UObject* FLevelSequenceBindingReference::Resolve(UObject* InContext) const
 {
-	if (PackageName.Len() == 0)
+	if (ExternalObjectPath.IsNull())
 	{
 		return FindObject<UObject>(InContext, *ObjectPath, false);
 	}
 	else
 	{
-		const TCHAR* SearchWithinPackage = *PackageName;
-		FString FixupPIEPackageName;
+		FSoftObjectPath TempPath = ExternalObjectPath;
 
 #if WITH_EDITORONLY_DATA
-		int32 PIEInstanceID = InContext ? InContext->GetOutermost()->PIEInstanceID : INDEX_NONE;
-
-		const FString ShortPackageOuterAndName = FPackageName::GetLongPackageAssetName(PackageName);
-		if (ensureMsgf(!ShortPackageOuterAndName.StartsWith(PLAYWORLD_PACKAGE_PREFIX), TEXT("Detected PIE world prefix in level sequence binding - this should not happen")))
-		{
-			if (PIEInstanceID != INDEX_NONE)
-			{
-				FixupPIEPackageName = FPackageName::GetLongPackagePath(PackageName) / FString::Printf(PLAYWORLD_PACKAGE_PREFIX TEXT("_%d_"), PIEInstanceID) + ShortPackageOuterAndName;
-			}
-			else
-			{
-				const FString PlayOnConsolePackageName = FPackageName::FilenameToLongPackageName(FPaths::Combine(*FPaths::ProjectSavedDir(), *GEngine->PlayOnConsoleSaveDir));
-				const FString ConsoleName = FString(TEXT("PC"));
-				const FString Prefix = FString(PLAYWORLD_CONSOLE_BASE_PACKAGE_PREFIX) + ConsoleName;
-
-				FixupPIEPackageName = PlayOnConsolePackageName + FPackageName::GetLongPackagePath(PackageName) / Prefix + ShortPackageOuterAndName;
-			}
-			SearchWithinPackage = *FixupPIEPackageName;
-		}
+		TempPath.FixupForPIE();
 #endif
 
-		UPackage* Package = FindPackage(nullptr, SearchWithinPackage);
-		if (!Package)
-		{
-			//Package wasn't found, fall back to default path
-			Package = FindPackage(nullptr, *PackageName);
-		}
-		return Package ? FindObject<UObject>(Package, *ObjectPath, false) : nullptr;
+		return TempPath.ResolveObject();
+	}
+}
+
+void FLevelSequenceBindingReference::PostSerialize(const FArchive& Ar)
+{
+	if (Ar.IsLoading() && !PackageName_DEPRECATED.IsEmpty())
+	{
+		// This was saved as two strings, combine into one soft object path so it handles PIE and redirectors properly
+		FString FullPath = PackageName_DEPRECATED + TEXT(".") + ObjectPath;
+
+		ExternalObjectPath.SetPath(FullPath);
+		ObjectPath.Reset();
+		PackageName_DEPRECATED.Reset();
 	}
 }
 
diff --git a/Engine/Source/Runtime/LevelSequence/Public/LevelSequenceActor.h b/Engine/Source/Runtime/LevelSequence/Public/LevelSequenceActor.h
--- a/Engine/Source/Runtime/LevelSequence/Public/LevelSequenceActor.h
+++ b/Engine/Source/Runtime/LevelSequence/Public/LevelSequenceActor.h
@@ -5,9 +5,9 @@
 #include "CoreMinimal.h"
 #include "UObject/ObjectMacros.h"
 #include "UObject/Object.h"
-#include "Misc/StringAssetReference.h"
+#include "UObject/SoftObjectPath.h"
 #include "GameFramework/Actor.h"
-#include "Misc/StringClassReference.h"
+#include "UObject/SoftObjectPath.h"
 #include "LevelSequencePlayer.h"
 #include "MovieSceneBindingOwnerInterface.h"
 #include "MovieSceneBindingOverrides.h"
@@ -38,7 +38,7 @@ public:
 	bool bUseBurnIn;
 
 	UPROPERTY(EditAnywhere, Category="General", meta=(EditCondition=bUseBurnIn, MetaClass="LevelSequenceBurnIn"))
-	FStringClassReference BurnInClass;
+	FSoftClassPath BurnInClass;
 
 	UPROPERTY(Instanced, EditAnywhere, Category="General", meta=(EditCondition=bUseBurnIn))
 	ULevelSequenceBurnInInitSettings* Settings;
@@ -77,7 +77,7 @@ public:
 	ULevelSequencePlayer* SequencePlayer;
 
 	UPROPERTY(EditAnywhere, BlueprintReadOnly, Category="General", meta=(AllowedClasses="LevelSequence"))
-	FStringAssetReference LevelSequence;
+	FSoftObjectPath LevelSequence;
 
 	UPROPERTY(EditAnywhere, BlueprintReadOnly, AdvancedDisplay, Category="General")
 	TArray<AActor*> AdditionalEventReceivers;
diff --git a/Engine/Source/Runtime/LevelSequence/Public/LevelSequenceBindingReference.h b/Engine/Source/Runtime/LevelSequence/Public/LevelSequenceBindingReference.h
--- a/Engine/Source/Runtime/LevelSequence/Public/LevelSequenceBindingReference.h
+++ b/Engine/Source/Runtime/LevelSequence/Public/LevelSequenceBindingReference.h
@@ -4,6 +4,7 @@
 
 #include "CoreMinimal.h"
 #include "UObject/ObjectMacros.h"
+#include "UObject/SoftObjectPath.h"
 #include "Misc/Guid.h"
 #include "Engine/Engine.h"
 #include "Paths.h"
@@ -39,17 +40,34 @@ struct FLevelSequenceBindingReference
 	 */
 	LEVELSEQUENCE_API UObject* Resolve(UObject* InContext) const;
 
+	/** Handles ExternalObjectPath fixup */
+	void PostSerialize(const FArchive& Ar);
+
 private:
 
-	/** The package name in which the object resides */
+	/** Replaced by ExternalObjectPath */
+	UPROPERTY()
+	FString PackageName_DEPRECATED;
+
+	/** Path to a specific actor/component inside an external package */
 	UPROPERTY()
-	FString PackageName;
+	FSoftObjectPath ExternalObjectPath;
 
-	/** The path of the object within the package, or from its outer (where PackageName is empty) */
+	/** Object path relative to a passed in context object, this is used if ExternalObjectPath is invalid */
 	UPROPERTY()
 	FString ObjectPath;
 };
 
+
+template<>
+struct TStructOpsTypeTraits<FLevelSequenceBindingReference> : public TStructOpsTypeTraitsBase2<FLevelSequenceBindingReference>
+{
+	enum
+	{
+		WithPostSerialize = true,
+	};
+};
+
 /**
  * An array of binding references
  */
diff --git a/Engine/Source/Runtime/MovieSceneCapture/Private/MovieSceneCapture.cpp b/Engine/Source/Runtime/MovieSceneCapture/Private/MovieSceneCapture.cpp
--- a/Engine/Source/Runtime/MovieSceneCapture/Private/MovieSceneCapture.cpp
+++ b/Engine/Source/Runtime/MovieSceneCapture/Private/MovieSceneCapture.cpp
@@ -36,7 +36,6 @@ FMovieSceneCaptureSettings::FMovieSceneCaptureSettings()
 	OutputDirectory.Path = FPaths::VideoCaptureDir();
 	FPaths::MakePlatformFilename( OutputDirectory.Path );
 
-	bCreateTemporaryCopiesOfLevels = false;
 	bUseRelativeFrameNumbers = false;
 	HandleFrames = 0;
 	GameModeOverride = nullptr;
diff --git a/Engine/Source/Runtime/MovieSceneCapture/Public/AutomatedLevelSequenceCapture.h b/Engine/Source/Runtime/MovieSceneCapture/Public/AutomatedLevelSequenceCapture.h
--- a/Engine/Source/Runtime/MovieSceneCapture/Public/AutomatedLevelSequenceCapture.h
+++ b/Engine/Source/Runtime/MovieSceneCapture/Public/AutomatedLevelSequenceCapture.h
@@ -4,7 +4,7 @@
 
 #include "CoreMinimal.h"
 #include "UObject/ObjectMacros.h"
-#include "Misc/StringAssetReference.h"
+#include "UObject/SoftObjectPath.h"
 #include "LevelSequencePlayer.h"
 #include "MovieSceneCapture.h"
 #include "AutomatedLevelSequenceCapture.generated.h"
@@ -119,7 +119,7 @@ private:
 
 	/** A level sequence asset to playback at runtime - used where the level sequence does not already exist in the world. */
 	UPROPERTY()
-	FStringAssetReference LevelSequenceAsset;
+	FSoftObjectPath LevelSequenceAsset;
 
 	/** The pre-existing level sequence actor to use for capture that specifies playback settings */
 	UPROPERTY()
diff --git a/Engine/Source/Runtime/MovieSceneCapture/Public/MovieSceneCaptureSettings.h b/Engine/Source/Runtime/MovieSceneCapture/Public/MovieSceneCaptureSettings.h
--- a/Engine/Source/Runtime/MovieSceneCapture/Public/MovieSceneCaptureSettings.h
+++ b/Engine/Source/Runtime/MovieSceneCapture/Public/MovieSceneCaptureSettings.h
@@ -36,10 +36,6 @@ struct MOVIESCENECAPTURE_API FMovieSceneCaptureSettings
 	UPROPERTY(config, EditAnywhere, Category=General, meta=(RelativePath))
 	FDirectoryPath OutputDirectory;
 
-	/** Whether to save temporary copies of all of the levels before capturing the movie.  This allows you to record movies of temporary work, or work that isn't yet saved, but it will take much longer for capturing to begin. */
-	UPROPERTY(config, EditAnywhere, Category=General, AdvancedDisplay)
-	bool bCreateTemporaryCopiesOfLevels;
-
 	/** Optional game mode to override the map's default game mode with.  This can be useful if the game's normal mode displays UI elements or loading screens that you don't want captured. */
 	UPROPERTY(config, EditAnywhere, Category=General, AdvancedDisplay)
 	TSubclassOf<class AGameModeBase> GameModeOverride;
diff --git a/Engine/Source/Runtime/MovieSceneCapture/Public/Protocols/CompositionGraphCaptureProtocol.h b/Engine/Source/Runtime/MovieSceneCapture/Public/Protocols/CompositionGraphCaptureProtocol.h
--- a/Engine/Source/Runtime/MovieSceneCapture/Public/Protocols/CompositionGraphCaptureProtocol.h
+++ b/Engine/Source/Runtime/MovieSceneCapture/Public/Protocols/CompositionGraphCaptureProtocol.h
@@ -4,7 +4,7 @@
 
 #include "CoreMinimal.h"
 #include "UObject/ObjectMacros.h"
-#include "Misc/StringAssetReference.h"
+#include "UObject/SoftObjectPath.h"
 #include "IMovieSceneCaptureProtocol.h"
 #include "MovieSceneCaptureProtocolSettings.h"
 #include "CompositionGraphCaptureProtocol.generated.h"
@@ -62,7 +62,7 @@ public:
 
 	/** Custom post processing material to use for rendering */
 	UPROPERTY(config, EditAnywhere, Category="Composition Graph Options", meta=(AllowedClasses=""))
-	FStringAssetReference PostProcessingMaterial;
+	FSoftObjectPath PostProcessingMaterial;
 };
 
 struct MOVIESCENECAPTURE_API FCompositionGraphCaptureProtocol : IMovieSceneCaptureProtocol
diff --git a/Engine/Source/Runtime/MovieSceneTracks/Private/Sections/MovieSceneEventSection.cpp b/Engine/Source/Runtime/MovieSceneTracks/Private/Sections/MovieSceneEventSection.cpp
--- a/Engine/Source/Runtime/MovieSceneTracks/Private/Sections/MovieSceneEventSection.cpp
+++ b/Engine/Source/Runtime/MovieSceneTracks/Private/Sections/MovieSceneEventSection.cpp
@@ -45,14 +45,14 @@ public:
 	}
 
 	/** Serialize an asset ptr */
-	virtual FArchive& operator<<(FAssetPtr& AssetPtr) override
+	virtual FArchive& operator<<(FSoftObjectPtr& AssetPtr) override
 	{
-		FStringAssetReference Ref = AssetPtr.ToStringReference();
+		FSoftObjectPath Ref = AssetPtr.ToSoftObjectPath();
 		*this << Ref;
 
 		if (IsLoading())
 		{
-			AssetPtr = FAssetPtr(Ref);
+			AssetPtr = FSoftObjectPtr(Ref);
 		}
 
 		return *this;
diff --git a/Engine/Source/Runtime/MovieSceneTracks/Public/Sections/MovieSceneEventSection.h b/Engine/Source/Runtime/MovieSceneTracks/Public/Sections/MovieSceneEventSection.h
--- a/Engine/Source/Runtime/MovieSceneTracks/Public/Sections/MovieSceneEventSection.h
+++ b/Engine/Source/Runtime/MovieSceneTracks/Public/Sections/MovieSceneEventSection.h
@@ -11,7 +11,7 @@
 #include "Curves/CurveInterface.h"
 #include "UObject/StructOnScope.h"
 #include "Engine/Engine.h"
-#include "StringAssetReference.h"
+#include "UObject/SoftObjectPath.h"
 #include "MovieSceneEventSection.generated.h"
 
 struct EventData;
@@ -84,8 +84,8 @@ struct MOVIESCENETRACKS_API FMovieSceneEventParameters
 
 private:
 
-	/** String asset reference to the type of this parameter payload */
-	FStringAssetReference StructType;
+	/** Soft object path to the type of this parameter payload */
+	FSoftObjectPath StructType;
 
 	/** Serialized bytes that represent the payload. Serialized internally with FEventParameterArchive */
 	TArray<uint8> StructBytes;
diff --git a/Engine/Source/Runtime/Niagara/Private/NiagaraModule.cpp b/Engine/Source/Runtime/Niagara/Private/NiagaraModule.cpp
--- a/Engine/Source/Runtime/Niagara/Private/NiagaraModule.cpp
+++ b/Engine/Source/Runtime/Niagara/Private/NiagaraModule.cpp
@@ -146,18 +146,18 @@ void FNiagaraTypeDefinition::RecreateUserDefinedTypeRegistry()
 
 	const UNiagaraSettings* Settings = GetDefault<UNiagaraSettings>();
 	check(Settings);
-	TArray<FStringAssetReference> TotalStructAssets;
-	for (FStringAssetReference AssetRef : Settings->AdditionalParameterTypes)
+	TArray<FSoftObjectPath> TotalStructAssets;
+	for (FSoftObjectPath AssetRef : Settings->AdditionalParameterTypes)
 	{
 		TotalStructAssets.AddUnique(AssetRef);
 	}
 
-	for (FStringAssetReference AssetRef : Settings->AdditionalPayloadTypes)
+	for (FSoftObjectPath AssetRef : Settings->AdditionalPayloadTypes)
 	{
 		TotalStructAssets.AddUnique(AssetRef);
 	}
 
-	for (FStringAssetReference AssetRef : TotalStructAssets)
+	for (FSoftObjectPath AssetRef : TotalStructAssets)
 	{
 		UObject* Obj = AssetRef.ResolveObject();
 		if (Obj == nullptr)
@@ -167,8 +167,8 @@ void FNiagaraTypeDefinition::RecreateUserDefinedTypeRegistry()
 
 		if (Obj != nullptr)
 		{
-			const FStringAssetReference* ParamRefFound = Settings->AdditionalParameterTypes.FindByPredicate([&](const FStringAssetReference& Ref) { return Ref.ToString() == AssetRef.ToString(); });
-			const FStringAssetReference* PayloadRefFound = Settings->AdditionalPayloadTypes.FindByPredicate([&](const FStringAssetReference& Ref) { return Ref.ToString() == AssetRef.ToString(); });
+			const FSoftObjectPath* ParamRefFound = Settings->AdditionalParameterTypes.FindByPredicate([&](const FSoftObjectPath& Ref) { return Ref.ToString() == AssetRef.ToString(); });
+			const FSoftObjectPath* PayloadRefFound = Settings->AdditionalPayloadTypes.FindByPredicate([&](const FSoftObjectPath& Ref) { return Ref.ToString() == AssetRef.ToString(); });
 			UScriptStruct* ScriptStruct = Cast<UScriptStruct>(Obj);
 			if (ScriptStruct != nullptr)
 			{
diff --git a/Engine/Source/Runtime/Niagara/Public/NiagaraSettings.h b/Engine/Source/Runtime/Niagara/Public/NiagaraSettings.h
--- a/Engine/Source/Runtime/Niagara/Public/NiagaraSettings.h
+++ b/Engine/Source/Runtime/Niagara/Public/NiagaraSettings.h
@@ -14,21 +14,21 @@ class NIAGARA_API UNiagaraSettings : public UDeveloperSettings
 		
 	/** Effect to duplicate as the base of all new effect assets created. */
 	UPROPERTY(config, EditAnywhere, Category = Niagara)
-	FStringAssetReference DefaultEffect;
+	FSoftObjectPath DefaultEffect;
 
 	/** Emitter to duplicate as the base of all new emitter assets created. */
 	UPROPERTY(config, EditAnywhere, Category = Niagara)
-	FStringAssetReference DefaultEmitter;
+	FSoftObjectPath DefaultEmitter;
 
 	/** Niagara script to duplicate as the base of all new script assets created. */
 	UPROPERTY(config, EditAnywhere, Category = Niagara)
-	FStringAssetReference DefaultScript;
+	FSoftObjectPath DefaultScript;
 
 	UPROPERTY(config, EditAnywhere, Category = Niagara, meta = (AllowedClasses = "ScriptStruct"))
-	TArray<FStringAssetReference> AdditionalParameterTypes;
+	TArray<FSoftObjectPath> AdditionalParameterTypes;
 
 	UPROPERTY(config, EditAnywhere, Category = Niagara, meta = (AllowedClasses = "ScriptStruct"))
-	TArray<FStringAssetReference> AdditionalPayloadTypes;
+	TArray<FSoftObjectPath> AdditionalPayloadTypes;
 
 #if WITH_EDITOR
 	virtual void PostEditChangeProperty(FPropertyChangedEvent& PropertyChangedEvent) override;
diff --git a/Engine/Source/Runtime/Slate/Public/SlateSharedPCH.h b/Engine/Source/Runtime/Slate/Public/SlateSharedPCH.h
--- a/Engine/Source/Runtime/Slate/Public/SlateSharedPCH.h
+++ b/Engine/Source/Runtime/Slate/Public/SlateSharedPCH.h
@@ -245,9 +245,7 @@
 #include "Misc/PackageName.h"
 #include "Templates/Casts.h"
 #include "UObject/LazyObjectPtr.h"
-#include "UObject/PersistentObjectPtr.h"
-#include "UObject/AssetPtr.h"
-#include "Misc/StringAssetReference.h"
+#include "UObject/SoftObjectPtr.h"
 #include "UObject/Package.h"
 #include "Misc/WorldCompositionUtility.h"
 #include "UObject/UnrealType.h"
diff --git a/Engine/Source/Runtime/UMG/Private/Components/RichTextBlockDecorator.cpp b/Engine/Source/Runtime/UMG/Private/Components/RichTextBlockDecorator.cpp
--- a/Engine/Source/Runtime/UMG/Private/Components/RichTextBlockDecorator.cpp
+++ b/Engine/Source/Runtime/UMG/Private/Components/RichTextBlockDecorator.cpp
@@ -1,7 +1,7 @@
 // Copyright 1998-2017 Epic Games, Inc. All Rights Reserved.
 
 #include "Components/RichTextBlockDecorator.h"
-#include "Misc/StringAssetReference.h"
+#include "UObject/SoftObjectPtr.h"
 #include "Rendering/DrawElements.h"
 #include "Framework/Text/SlateTextRun.h"
 
@@ -149,7 +149,7 @@ void FDefaultRichTextDecorator::ExplodeRunInfo(const FRunInfo& InRunInfo, FSlate
 
 	if ( FontFamilyString )
 	{
-		FStringAssetReference Font(**FontFamilyString);
+		FSoftObjectPath Font(**FontFamilyString);
 		if ( UObject* FontAsset = Font.TryLoad() )
 		{
 			OutFont.FontObject = FontAsset;
diff --git a/Engine/Source/Runtime/UMG/Private/UserWidget.cpp b/Engine/Source/Runtime/UMG/Private/UserWidget.cpp
--- a/Engine/Source/Runtime/UMG/Private/UserWidget.cpp
+++ b/Engine/Source/Runtime/UMG/Private/UserWidget.cpp
@@ -81,7 +81,8 @@ UWidgetBlueprintGeneratedClass* UUserWidget::GetWidgetTreeOwningClass()
 		// Force post load on the generated class so all subobjects are done (specifically the widget tree).
 		BGClass->ConditionalPostLoad();
 
-		const bool bNoRootWidget = ( nullptr == BGClass->WidgetTree ) || ( nullptr == BGClass->WidgetTree->RootWidget );
+		const bool bNoRootWidget = !BGClass->HasTemplate() && ( ( nullptr == BGClass->WidgetTree ) || ( nullptr == BGClass->WidgetTree->RootWidget ) );
+
 		if ( bNoRootWidget )
 		{
 			UWidgetBlueprintGeneratedClass* SuperBGClass = Cast<UWidgetBlueprintGeneratedClass>(BGClass->GetSuperClass());
diff --git a/Engine/Source/Runtime/UMG/Public/Blueprint/WidgetBlueprintGeneratedClass.h b/Engine/Source/Runtime/UMG/Public/Blueprint/WidgetBlueprintGeneratedClass.h
--- a/Engine/Source/Runtime/UMG/Public/Blueprint/WidgetBlueprintGeneratedClass.h
+++ b/Engine/Source/Runtime/UMG/Public/Blueprint/WidgetBlueprintGeneratedClass.h
@@ -129,7 +129,7 @@ private:
 	bool bTemplateInitialized;
 
 	UPROPERTY()
-	TAssetPtr<UUserWidget> TemplateAsset;
+	TSoftObjectPtr<UUserWidget> TemplateAsset;
 
 	UPROPERTY(Transient)
 	mutable UUserWidget* Template;
