Sequencer: Add back EditAnywhere to BurnIn Settings

diff --git a/Engine/Source/Runtime/LevelSequence/Public/LevelSequenceActor.h b/Engine/Source/Runtime/LevelSequence/Public/LevelSequenceActor.h
--- a/Engine/Source/Runtime/LevelSequence/Public/LevelSequenceActor.h
+++ b/Engine/Source/Runtime/LevelSequence/Public/LevelSequenceActor.h
@@ -44,7 +44,7 @@ public:
 	UPROPERTY(config, EditAnywhere, BlueprintReadWrite, Category="General", meta=(EditCondition=bUseBurnIn, MetaClass="LevelSequenceBurnIn"))
 	FSoftClassPath BurnInClass;
 
-	UPROPERTY(Instanced, BlueprintReadWrite, Category="General", meta=(EditCondition=bUseBurnIn))
+	UPROPERTY(Instanced, EditAnywhere, BlueprintReadWrite, Category="General", meta=(EditCondition=bUseBurnIn))
 	ULevelSequenceBurnInInitSettings* Settings;
 
 protected:
