Sequencer: Expose DefaultLevelSequenceInstanceData

diff --git a/Engine/Source/Runtime/LevelSequence/Public/DefaultLevelSequenceInstanceData.h b/Engine/Source/Runtime/LevelSequence/Public/DefaultLevelSequenceInstanceData.h
--- a/Engine/Source/Runtime/LevelSequence/Public/DefaultLevelSequenceInstanceData.h
+++ b/Engine/Source/Runtime/LevelSequence/Public/DefaultLevelSequenceInstanceData.h
@@ -11,7 +11,7 @@
 
 /** Default instance data class that level sequences understand. Implements IMovieSceneTransformOrigin. */
 UCLASS(BlueprintType)
-class UDefaultLevelSequenceInstanceData
+class LEVELSEQUENCE_API UDefaultLevelSequenceInstanceData
 	: public UObject
 	, public IMovieSceneTransformOrigin
 {
