Sequencer: Protect against trying to key based on the CDO changing

diff --git a/Engine/Source/Runtime/LevelSequence/Private/LevelSequence.cpp b/Engine/Source/Runtime/LevelSequence/Private/LevelSequence.cpp
--- a/Engine/Source/Runtime/LevelSequence/Private/LevelSequence.cpp
+++ b/Engine/Source/Runtime/LevelSequence/Private/LevelSequence.cpp
@@ -284,7 +284,10 @@ UObject* ULevelSequence::GetParentObject(UObject* Object) const
 
 	if (UAnimInstance* AnimInstance = Cast<UAnimInstance>(Object))
 	{
-		return AnimInstance->GetOwningComponent();
+		if (AnimInstance->GetWorld())
+		{
+			return AnimInstance->GetOwningComponent();
+		}
 	}
 
 	return nullptr;
