Check to see if level sequence has itself as a tick priority actor.

diff --git a/Engine/Source/Runtime/LevelSequence/Private/LevelSequencePlayer.cpp b/Engine/Source/Runtime/LevelSequence/Private/LevelSequencePlayer.cpp
--- a/Engine/Source/Runtime/LevelSequence/Private/LevelSequencePlayer.cpp
+++ b/Engine/Source/Runtime/LevelSequence/Private/LevelSequencePlayer.cpp
@@ -222,6 +222,12 @@ void ULevelSequencePlayer::NotifyBindingUpdate(const FGuid& InGuid, FMovieSceneS
 	{
 		if (AActor* Actor = Cast<AActor>(WeakObject.Get()))
 		{
+			
+			if (Actor == LevelSequenceActor)
+			{
+				continue;
+			}
+
 			for (UActorComponent* Component : Actor->GetComponents())
 			{
 				if (Component)
