Sequencer: Reset the last view target when stopping the sequence so that it gets recached when evaluating the sequence again.

diff --git a/Engine/Source/Runtime/LevelSequence/Private/LevelSequencePlayer.cpp b/Engine/Source/Runtime/LevelSequence/Private/LevelSequencePlayer.cpp
--- a/Engine/Source/Runtime/LevelSequence/Private/LevelSequencePlayer.cpp
+++ b/Engine/Source/Runtime/LevelSequence/Private/LevelSequencePlayer.cpp
@@ -109,6 +109,7 @@ void ULevelSequencePlayer::OnStopped()
 		}
 	}
 	PrerequisiteActors.Reset();
+	LastViewTarget.Reset();
 }
 
 /* IMovieScenePlayer interface
