Sequencer: Fix CurrentShot LocalTime computation by using sequence time in playback resolution to compute the local shot time.

Also, fixed the burnin asset so that CurrentShotLocalTime is hooked up to ShotFrame instead of MasterTime.

This fixes a bug where the burnin's {ShotFrame} is not reporting the local shot frame number.

Copy from Release-4.20

diff --git a/Engine/Source/Runtime/LevelSequence/Private/LevelSequencePlayer.cpp b/Engine/Source/Runtime/LevelSequence/Private/LevelSequencePlayer.cpp
--- a/Engine/Source/Runtime/LevelSequence/Private/LevelSequencePlayer.cpp
+++ b/Engine/Source/Runtime/LevelSequence/Private/LevelSequencePlayer.cpp
@@ -349,11 +349,13 @@ void ULevelSequencePlayer::TakeFrameSnapshot(FLevelSequencePlayerSnapshot& OutSn
 		{
 			// Assume that shots with no sequence start at 0.
 			FMovieSceneSequenceTransform OuterToInnerTransform = ActiveShot->OuterToInnerTransform();
-			UMovieSceneSequence*         InnerSequence         = ActiveShot->GetSequence();
-			FFrameRate                   InnerFrameRate        = InnerSequence ? InnerSequence->GetMovieScene()->GetDisplayRate() : PlayPosition.GetInputRate();
+			UMovieSceneSequence*         InnerSequence = ActiveShot->GetSequence();
+			FFrameRate                   InnerTickResoloution = InnerSequence ? InnerSequence->GetMovieScene()->GetTickResolution() : PlayPosition.GetOutputRate();
+			FFrameRate                   InnerFrameRate = InnerSequence ? InnerSequence->GetMovieScene()->GetDisplayRate() : PlayPosition.GetInputRate();
+			FFrameTime                   InnerDisplayTime = ConvertFrameTime(CurrentSequenceTime * OuterToInnerTransform, InnerTickResoloution, InnerFrameRate);
 
-			OutSnapshot.CurrentShotName      = ActiveShot->GetShotDisplayName();
-			OutSnapshot.CurrentShotLocalTime = FQualifiedFrameTime(CurrentPlayTime * OuterToInnerTransform, InnerFrameRate);
+			OutSnapshot.CurrentShotName = ActiveShot->GetShotDisplayName();
+			OutSnapshot.CurrentShotLocalTime = FQualifiedFrameTime(InnerDisplayTime, InnerFrameRate);
 			OutSnapshot.ShotID = ActiveShot->GetSequenceID();
 		}
 	}
