UE-56235 - Sequencer Possessable references no longer lose their references when saving a map for the first time.

diff --git a/Engine/Source/Editor/UnrealEd/Private/FileHelpers.cpp b/Engine/Source/Editor/UnrealEd/Private/FileHelpers.cpp
--- a/Engine/Source/Editor/UnrealEd/Private/FileHelpers.cpp
+++ b/Engine/Source/Editor/UnrealEd/Private/FileHelpers.cpp
@@ -680,11 +680,6 @@ static bool SaveWorld(UWorld* World,
 			SaveErrors.Flush();
 		}
 
-		if ( bAddedAssetPathRedirection )
-		{
-			GRedirectCollector.RemoveAssetPathRedirection( *OldPath.GetAssetPathString() );
-		}
-
 		// @todo Autosaving should save build data as well
 		if (bSuccess && !bAutosaving)
 		{
diff --git a/Engine/Source/Runtime/LevelSequence/Private/LevelSequenceBindingReference.cpp b/Engine/Source/Runtime/LevelSequence/Private/LevelSequenceBindingReference.cpp
--- a/Engine/Source/Runtime/LevelSequence/Private/LevelSequenceBindingReference.cpp
+++ b/Engine/Source/Runtime/LevelSequence/Private/LevelSequenceBindingReference.cpp
@@ -51,6 +51,11 @@ UObject* FLevelSequenceBindingReference::Resolve(UObject* InContext) const
 	{
 		FSoftObjectPath TempPath = ExternalObjectPath;
 
+		// Soft Object Paths don't follow asset redirectors when attempting to call ResolveObject or TryLoad.
+		// We want to follow the asset redirector so that maps that have been renamed (from Untitled to their first asset name)
+		// properly resolve. This fixes Possessable bindings losing their references the first time you save a map.
+		TempPath.PreSavePath();
+
 #if WITH_EDITORONLY_DATA
 		int32 ContextPlayInEditorID = InContext ? InContext->GetOutermost()->PIEInstanceID : INDEX_NONE;
 
