Sequencer: Clarified allow bindings from asset tooltips

diff --git a/Engine/Source/Runtime/LevelSequence/Public/LevelSequenceActor.h b/Engine/Source/Runtime/LevelSequence/Public/LevelSequenceActor.h
--- a/Engine/Source/Runtime/LevelSequence/Public/LevelSequenceActor.h
+++ b/Engine/Source/Runtime/LevelSequence/Public/LevelSequenceActor.h
@@ -135,7 +135,12 @@ public:
 
 public:
 
-	/** Overrides the specified binding with the specified actors, optionally still allowing the bindings defined in the Level Sequence asset */
+	/** Overrides the specified binding with the specified actors, optionally still allowing the bindings defined in the Level Sequence asset
+	 *
+	 * @param Binding Binding to modify
+	 * @param Actors Actors to bind
+	 * @param bAllowBindingsFromAsset Allow bindings from the level sequence asset
+	 */
 	UFUNCTION(BlueprintCallable, Category="Game|Cinematic|Bindings")
 	void SetBinding(FMovieSceneObjectBindingID Binding, const TArray<AActor*>& Actors, bool bAllowBindingsFromAsset = false)
 	{
@@ -146,7 +151,12 @@ public:
 		}
 	}
 
-	/** Adds the specified actor to the overridden bindings for the specified binding ID, optionally still allowing the bindings defined in the Level Sequence asset */
+	/** Adds the specified actor to the overridden bindings for the specified binding ID, optionally still allowing the bindings defined in the Level Sequence asset
+	 *
+	 * @param Binding Binding to modify
+	 * @param Actor Actor to bind
+	 * @param bAllowBindingsFromAsset Allow bindings from the level sequence asset
+	 */
 	UFUNCTION(BlueprintCallable, Category="Game|Cinematic|Bindings")
 	void AddBinding(FMovieSceneObjectBindingID Binding, AActor* Actor, bool bAllowBindingsFromAsset = false)
 	{
