Sequencer: Make unique spawnable name to avoid it being renamed in SetActorLabel.

diff --git a/Engine/Source/Runtime/LevelSequence/Private/LevelSequenceActorSpawner.cpp b/Engine/Source/Runtime/LevelSequence/Private/LevelSequenceActorSpawner.cpp
--- a/Engine/Source/Runtime/LevelSequence/Private/LevelSequenceActorSpawner.cpp
+++ b/Engine/Source/Runtime/LevelSequence/Private/LevelSequenceActorSpawner.cpp
@@ -71,11 +71,18 @@ UObject* FLevelSequenceActorSpawner::SpawnObject(FMovieSceneSpawnable& Spawnable
 		UE_LOG(LogMovieScene, Warning, TEXT("Can't find sublevel '%s' to spawn '%s' into"), *Spawnable.GetLevelName().ToString(), *Spawnable.GetName());
 	}
 
+	// Construct the object with the same name that we will set later on the actor to avoid renaming it inside SetActorLabel
+	FName SpawnName =
+#if WITH_EDITOR
+		MakeUniqueObjectName(WorldContext->PersistentLevel, ObjectTemplate->GetClass(), *Spawnable.GetName());
+#else
+		NAME_None;
+#endif
 
 	// Spawn the puppet actor
 	FActorSpawnParameters SpawnInfo;
 	{
-		SpawnInfo.Name = NAME_None;
+		SpawnInfo.Name = SpawnName;
 		SpawnInfo.ObjectFlags = ObjectFlags;
 		SpawnInfo.SpawnCollisionHandlingOverride = ESpawnActorCollisionHandlingMethod::AlwaysSpawn;
 		// @todo: Spawning with a non-CDO template is fraught with issues
@@ -138,14 +145,19 @@ UObject* FLevelSequenceActorSpawner::SpawnObject(FMovieSceneSpawnable& Spawnable
 		{
 			Component->SetFlags(RF_Transactional);
 		}
-	}
-
-	SpawnedActor->SetActorLabel(Spawnable.GetName());
+	}	
 #endif
 
 	const bool bIsDefaultTransform = true;
 	SpawnedActor->FinishSpawning(SpawnTransform, bIsDefaultTransform);
 
+#if WITH_EDITOR
+	if (GIsEditor)
+	{
+		SpawnedActor->SetActorLabel(Spawnable.GetName());
+	}
+#endif
+
 	return SpawnedActor;
 }
 
