If level sequence actor async loads after the sequence player is already gone, it would crash

diff --git a/Engine/Source/Runtime/LevelSequence/Private/LevelSequenceActor.cpp b/Engine/Source/Runtime/LevelSequence/Private/LevelSequenceActor.cpp
--- a/Engine/Source/Runtime/LevelSequence/Private/LevelSequenceActor.cpp
+++ b/Engine/Source/Runtime/LevelSequence/Private/LevelSequenceActor.cpp
@@ -233,7 +233,7 @@ void ALevelSequenceActor::OnSequenceLoaded(const FName& PackageName, UPackage* P
 	if (Result == EAsyncLoadingResult::Succeeded)
 	{
 		ULevelSequence* LevelSequenceAsset = GetSequence();
-		if (SequencePlayer->GetSequence() != LevelSequenceAsset)
+		if (SequencePlayer && SequencePlayer->GetSequence() != LevelSequenceAsset)
 		{
 			SequencePlayer->Initialize(LevelSequenceAsset, GetLevel(), PlaybackSettings);
 		}
