Sequencer: Adjust transform origin tooltip

diff --git a/Engine/Source/Runtime/LevelSequence/Public/DefaultLevelSequenceInstanceData.h b/Engine/Source/Runtime/LevelSequence/Public/DefaultLevelSequenceInstanceData.h
--- a/Engine/Source/Runtime/LevelSequence/Public/DefaultLevelSequenceInstanceData.h
+++ b/Engine/Source/Runtime/LevelSequence/Public/DefaultLevelSequenceInstanceData.h
@@ -32,7 +32,7 @@ public:
 	UPROPERTY(EditAnywhere, BlueprintReadWrite, Category="General")
 	AActor* TransformOriginActor;
 
-	/** Specifies a transform from which all absolute transform sections inside the sequence should be added to. Scale is ignored. */
+	/** Specifies a transform that offsets all absolute transform sections in this sequence. Will compound with attach tracks. Scale is ignored. Not applied to Relative or Additive sections. */
 	UPROPERTY(EditAnywhere, BlueprintReadWrite, Category="General")
 	FTransform TransformOrigin;
 };
