Sequencer: Added null check for object template

diff --git a/Engine/Source/Runtime/LevelSequence/Private/LevelSequenceSpawnRegister.cpp b/Engine/Source/Runtime/LevelSequence/Private/LevelSequenceSpawnRegister.cpp
--- a/Engine/Source/Runtime/LevelSequence/Private/LevelSequenceSpawnRegister.cpp
+++ b/Engine/Source/Runtime/LevelSequence/Private/LevelSequenceSpawnRegister.cpp
@@ -18,7 +18,7 @@ UObject* FLevelSequenceSpawnRegister::SpawnObject(FMovieSceneSpawnable& Spawnabl
 {
 	for (TSharedRef<IMovieSceneObjectSpawner> MovieSceneObjectSpawner : MovieSceneObjectSpawners)
 	{
-		if (Spawnable.GetObjectTemplate()->IsA(MovieSceneObjectSpawner->GetSupportedTemplateType()))
+		if (Spawnable.GetObjectTemplate() != nullptr && Spawnable.GetObjectTemplate()->IsA(MovieSceneObjectSpawner->GetSupportedTemplateType()))
 		{
 			
 			UObject* SpawnedObject = MovieSceneObjectSpawner->SpawnObject(Spawnable, TemplateID, Player);
