Sequencer: Check for world nullptr

diff --git a/Engine/Source/Runtime/LevelSequence/Private/LevelSequenceActor.cpp b/Engine/Source/Runtime/LevelSequence/Private/LevelSequenceActor.cpp
--- a/Engine/Source/Runtime/LevelSequence/Private/LevelSequenceActor.cpp
+++ b/Engine/Source/Runtime/LevelSequence/Private/LevelSequenceActor.cpp
@@ -217,7 +217,9 @@ void ALevelSequenceActor::InitializePlayer()
 {
 	ULevelSequence* LevelSequenceAsset = GetSequence(true, true);
 
-	if (GetWorld()->IsGameWorld() && (LevelSequenceAsset != nullptr))
+	UWorld* World = GetWorld();
+
+	if (World && World->IsGameWorld() && (LevelSequenceAsset != nullptr))
 	{
 		FMovieSceneSequencePlaybackSettings PlaybackSettingsCopy = PlaybackSettings;
 
@@ -228,7 +230,7 @@ void ALevelSequenceActor::InitializePlayer()
 		}
 
 		SequencePlayer = NewObject<ULevelSequencePlayer>(this, "AnimationPlayer");
-		SequencePlayer->Initialize(LevelSequenceAsset, GetWorld(), PlaybackSettingsCopy);
+		SequencePlayer->Initialize(LevelSequenceAsset, World, PlaybackSettingsCopy);
 		SequencePlayer->SetEventReceivers(TArray<UObject*>(AdditionalEventReceivers));
 
 		RefreshBurnIn();
