Sequencer: Move autoplay to playback settings so that CreateLevelSequencePlayer can be set in autoplay mode

diff --git a/Engine/Source/Editor/DetailCustomizations/Private/MovieSceneSequencePlaybackSettingsCustomization.cpp b/Engine/Source/Editor/DetailCustomizations/Private/MovieSceneSequencePlaybackSettingsCustomization.cpp
--- a/Engine/Source/Editor/DetailCustomizations/Private/MovieSceneSequencePlaybackSettingsCustomization.cpp
+++ b/Engine/Source/Editor/DetailCustomizations/Private/MovieSceneSequencePlaybackSettingsCustomization.cpp
@@ -132,6 +132,7 @@ void FMovieSceneSequencePlaybackSettingsCustomization::CustomizeChildren(TShared
 				]
 		];
 
+	ChildBuilder.AddProperty(StructPropertyHandle->GetChildHandle(GET_MEMBER_NAME_CHECKED(FMovieSceneSequencePlaybackSettings, bAutoPlay)).ToSharedRef());
 	ChildBuilder.AddProperty(StructPropertyHandle->GetChildHandle(GET_MEMBER_NAME_CHECKED(FMovieSceneSequencePlaybackSettings, bRandomStartTime)).ToSharedRef());
 	ChildBuilder.AddProperty(StructPropertyHandle->GetChildHandle(GET_MEMBER_NAME_CHECKED(FMovieSceneSequencePlaybackSettings, StartTime)).ToSharedRef());
 	ChildBuilder.AddProperty(StructPropertyHandle->GetChildHandle(GET_MEMBER_NAME_CHECKED(FMovieSceneSequencePlaybackSettings, bRestoreState)).ToSharedRef());
diff --git a/Engine/Source/Editor/MovieSceneTools/Private/AutomatedLevelSequenceCapture.cpp b/Engine/Source/Editor/MovieSceneTools/Private/AutomatedLevelSequenceCapture.cpp
--- a/Engine/Source/Editor/MovieSceneTools/Private/AutomatedLevelSequenceCapture.cpp
+++ b/Engine/Source/Editor/MovieSceneTools/Private/AutomatedLevelSequenceCapture.cpp
@@ -208,6 +208,7 @@ void UAutomatedLevelSequenceCapture::Initialize(TSharedPtr<FSceneViewport> InVie
 		Actor->PlaybackSettings.LoopCount = 0;
 		Actor->PlaybackSettings.TimeController = MakeShared<FMovieSceneTimeController_FrameStep>();
 		Actor->PlaybackSettings.bPauseAtEnd = true;
+		Actor->PlaybackSettings.bAutoPlay = false;
 
 		if (BurnInOptions)
 		{
@@ -231,7 +232,6 @@ void UAutomatedLevelSequenceCapture::Initialize(TSharedPtr<FSceneViewport> InVie
 			}
 			Actor->InitializePlayer();
 		}
-		Actor->bAutoPlay = false;
 
 		if (InitializeShots())
 		{
diff --git a/Engine/Source/Runtime/LevelSequence/Private/LevelSequenceActor.cpp b/Engine/Source/Runtime/LevelSequence/Private/LevelSequenceActor.cpp
--- a/Engine/Source/Runtime/LevelSequence/Private/LevelSequenceActor.cpp
+++ b/Engine/Source/Runtime/LevelSequence/Private/LevelSequenceActor.cpp
@@ -51,7 +51,7 @@ ALevelSequenceActor::ALevelSequenceActor(const FObjectInitializer& Init)
 	bOverrideInstanceData = false;
 
 	PrimaryActorTick.bCanEverTick = true;
-	bAutoPlay = false;
+	bAutoPlay_DEPRECATED = false;
 }
 
 void ALevelSequenceActor::PostInitializeComponents()
@@ -110,6 +110,13 @@ void ALevelSequenceActor::PostLoad()
 {
 	Super::PostLoad();
 
+	// If autoplay was previously enabled, initialize the playback settings to autoplay
+	if (bAutoPlay_DEPRECATED)
+	{
+		PlaybackSettings.bAutoPlay = bAutoPlay_DEPRECATED;
+		bAutoPlay_DEPRECATED = false;
+	}
+
 #if WITH_EDITORONLY_DATA
 	// Fix sprite component so that it's attached to the root component. In the past, the sprite component was the root component.
 	UBillboardComponent* SpriteComponent = FindComponentByClass<UBillboardComponent>();
@@ -233,7 +240,7 @@ void ALevelSequenceActor::InitializePlayer()
 
 		RefreshBurnIn();
 
-		if (bAutoPlay)
+		if (PlaybackSettingsCopy.bAutoPlay)
 		{
 			SequencePlayer->Play();
 		}
diff --git a/Engine/Source/Runtime/LevelSequence/Public/LevelSequenceActor.h b/Engine/Source/Runtime/LevelSequence/Public/LevelSequenceActor.h
--- a/Engine/Source/Runtime/LevelSequence/Public/LevelSequenceActor.h
+++ b/Engine/Source/Runtime/LevelSequence/Public/LevelSequenceActor.h
@@ -71,8 +71,8 @@ public:
 
 public:
 
-	UPROPERTY(EditAnywhere, BlueprintReadOnly, Category="Playback")
-	bool bAutoPlay;
+	UPROPERTY()
+	bool bAutoPlay_DEPRECATED;
 
 	UPROPERTY(EditAnywhere, BlueprintReadOnly, Category="Playback", meta=(ShowOnlyInnerProperties))
 	FMovieSceneSequencePlaybackSettings PlaybackSettings;
diff --git a/Engine/Source/Runtime/MovieScene/Public/MovieSceneSequencePlayer.h b/Engine/Source/Runtime/MovieScene/Public/MovieSceneSequencePlayer.h
--- a/Engine/Source/Runtime/MovieScene/Public/MovieSceneSequencePlayer.h
+++ b/Engine/Source/Runtime/MovieScene/Public/MovieSceneSequencePlayer.h
@@ -35,7 +35,8 @@ USTRUCT(BlueprintType)
 struct FMovieSceneSequencePlaybackSettings
 {
 	FMovieSceneSequencePlaybackSettings()
-		: LoopCount(0)
+		: bAutoPlay(false)
+		, LoopCount(0)
 		, PlayRate(1.f)
 		, bRandomStartTime(false)
 		, StartTime(0.f)
@@ -51,6 +52,10 @@ struct FMovieSceneSequencePlaybackSettings
 
 	GENERATED_BODY()
 
+	/** Auto-play the sequence when created */
+	UPROPERTY(EditAnywhere, BlueprintReadWrite, Category = "Playback")
+	bool bAutoPlay;
+
 	/** Number of times to loop playback. -1 for infinite, else the number of times to loop before stopping */
 	UPROPERTY(EditAnywhere, BlueprintReadWrite, Category="Playback", meta=(UIMin=1, DisplayName="Loop"))
 	int32 LoopCount;
