When a sequence stops playing set client simulating view target back to false.

diff --git a/Engine/Source/Runtime/LevelSequence/Private/LevelSequencePlayer.cpp b/Engine/Source/Runtime/LevelSequence/Private/LevelSequencePlayer.cpp
--- a/Engine/Source/Runtime/LevelSequence/Private/LevelSequencePlayer.cpp
+++ b/Engine/Source/Runtime/LevelSequence/Private/LevelSequencePlayer.cpp
@@ -108,6 +108,20 @@ void ULevelSequencePlayer::OnStopped()
 			Actor->PrimaryActorTick.RemovePrerequisite(LevelSequenceActor, LevelSequenceActor->PrimaryActorTick);
 		}
 	}
+
+	if (World != nullptr && World->GetGameInstance() != nullptr)
+	{
+		APlayerController* PC = World->GetGameInstance()->GetFirstLocalPlayerController();
+
+		if (PC != nullptr)
+		{
+			if (PC->PlayerCameraManager)
+			{
+				PC->PlayerCameraManager->bClientSimulatingViewTarget = false;
+			}
+		}
+	}
+
 	PrerequisiteActors.Reset();
 	LastViewTarget.Reset();
 }
