Sequencer: Skip updating camera cuts if camera to set is null and the last active camera and the current active camera is the same

diff --git a/Engine/Source/Runtime/LevelSequence/Private/LevelSequencePlayer.cpp b/Engine/Source/Runtime/LevelSequence/Private/LevelSequencePlayer.cpp
--- a/Engine/Source/Runtime/LevelSequence/Private/LevelSequencePlayer.cpp
+++ b/Engine/Source/Runtime/LevelSequence/Private/LevelSequencePlayer.cpp
@@ -198,6 +198,12 @@ void ULevelSequencePlayer::UpdateCameraCut(UObject* CameraObject, UObject* Unloc
 	if (CameraActor == nullptr)
 	{
 		CameraActor = LastViewTarget.Get();
+
+		// Skip if the last view target is the same as the current view target so that there's no additional camera cut
+		if (CameraActor == ViewTarget)
+		{
+			return;
+		}
 	}
 
 	FViewTargetTransitionParams TransitionParams;
