Sequencer: Defer construction for CreateLevelSequencePlayer() so that BeginPlay is called for autoplay. There used to be an explicit call to autoplay, but this is now all handled in BeginPlay().

diff --git a/Engine/Source/Runtime/LevelSequence/Private/LevelSequencePlayer.cpp b/Engine/Source/Runtime/LevelSequence/Private/LevelSequencePlayer.cpp
--- a/Engine/Source/Runtime/LevelSequence/Private/LevelSequencePlayer.cpp
+++ b/Engine/Source/Runtime/LevelSequence/Private/LevelSequencePlayer.cpp
@@ -52,6 +52,10 @@ ULevelSequencePlayer* ULevelSequencePlayer::CreateLevelSequencePlayer(UObject* W
 	SpawnParams.SpawnCollisionHandlingOverride = ESpawnActorCollisionHandlingMethod::AlwaysSpawn;
 	SpawnParams.ObjectFlags |= RF_Transient;
 	SpawnParams.bAllowDuringConstructionScript = true;
+
+	// Defer construction for autoplay so that BeginPlay() is called
+	SpawnParams.bDeferConstruction = true;
+
 	ALevelSequenceActor* Actor = World->SpawnActor<ALevelSequenceActor>(SpawnParams);
 
 	Actor->PlaybackSettings = Settings;
@@ -60,6 +64,9 @@ ULevelSequencePlayer* ULevelSequencePlayer::CreateLevelSequencePlayer(UObject* W
 	Actor->InitializePlayer();
 	OutActor = Actor;
 
+	FTransform DefaultTransform;
+	Actor->FinishSpawning(DefaultTransform);
+
 	return Actor->SequencePlayer;
 }
 
