Sequencer: Bindings that exist within instance sub levels now work correctly

  - On playback we now cache off both the original asset path name, and the level ptr if we're playing back within a level that has a ULevelStreaming associated with it
  - When resolving bindings, we compare the original binding's asset path name to the one we're playing back within. If it matches, we resolve using the ULevel ptr provided, rather than the UWorld.
  - Ultimately we would like to change this all to make the context _always_ a ULevel rather than a UWorld, but that is a high risk change to make since a lot of code (mostly editor code) heavily relies on the playback context being a UWorld.

diff --git a/Engine/Plugins/Experimental/ControlRig/Source/ControlRigEditor/Private/Sequencer/ControlRigSequenceExporter.cpp b/Engine/Plugins/Experimental/ControlRig/Source/ControlRigEditor/Private/Sequencer/ControlRigSequenceExporter.cpp
--- a/Engine/Plugins/Experimental/ControlRig/Source/ControlRigEditor/Private/Sequencer/ControlRigSequenceExporter.cpp
+++ b/Engine/Plugins/Experimental/ControlRig/Source/ControlRigEditor/Private/Sequencer/ControlRigSequenceExporter.cpp
@@ -213,7 +213,7 @@ void Convert(UControlRigSequence* Sequence, UAnimSequence* AnimSequence, USkelet
 			LevelSequenceActor->LevelSequence = Sequence;
 			LevelSequenceActor->PlaybackSettings.bRestoreState = true;
 			LevelSequenceActor->SequencePlayer = NewObject<ULevelSequencePlayer>(LevelSequenceActor, "AnimationPlayer");
-			LevelSequenceActor->SequencePlayer->Initialize(Sequence, World, LevelSequenceActor->PlaybackSettings);
+			LevelSequenceActor->SequencePlayer->Initialize(Sequence, LevelSequenceActor->GetLevel(), LevelSequenceActor->PlaybackSettings);
 
 			// Now set up our animation sequence
 			AnimSequence->RecycleAnimSequence();
diff --git a/Engine/Source/Runtime/LevelSequence/Private/LevelSequence.cpp b/Engine/Source/Runtime/LevelSequence/Private/LevelSequence.cpp
--- a/Engine/Source/Runtime/LevelSequence/Private/LevelSequence.cpp
+++ b/Engine/Source/Runtime/LevelSequence/Private/LevelSequence.cpp
@@ -287,6 +287,11 @@ bool ULevelSequence::CanPossessObject(UObject& Object, UObject* InPlaybackContex
 }
 
 void ULevelSequence::LocateBoundObjects(const FGuid& ObjectId, UObject* Context, TArray<UObject*, TInlineAllocator<1>>& OutObjects) const
+{
+	LocateBoundObjects(ObjectId, Context, NAME_None, OutObjects);
+}
+
+void ULevelSequence::LocateBoundObjects(const FGuid& ObjectId, UObject* Context, FName StreamedLevelAssetPath, TArray<UObject*, TInlineAllocator<1>>& OutObjects) const
 {
 	// Handle legacy object references
 	UObject* Object = Context ? ObjectReferences.ResolveBinding(ObjectId, Context) : nullptr;
@@ -295,7 +300,7 @@ void ULevelSequence::LocateBoundObjects(const FGuid& ObjectId, UObject* Context,
 		OutObjects.Add(Object);
 	}
 
-	BindingReferences.ResolveBinding(ObjectId, Context, OutObjects);
+	BindingReferences.ResolveBinding(ObjectId, Context, StreamedLevelAssetPath, OutObjects);
 }
 
 void ULevelSequence::GatherExpiredObjects(const FMovieSceneObjectCache& InObjectCache, TArray<FGuid>& OutInvalidIDs) const
diff --git a/Engine/Source/Runtime/LevelSequence/Private/LevelSequenceActor.cpp b/Engine/Source/Runtime/LevelSequence/Private/LevelSequenceActor.cpp
--- a/Engine/Source/Runtime/LevelSequence/Private/LevelSequenceActor.cpp
+++ b/Engine/Source/Runtime/LevelSequence/Private/LevelSequenceActor.cpp
@@ -181,7 +181,7 @@ void ALevelSequenceActor::SetSequence(ULevelSequence* InSequence)
 		// cbb: should ideally null out the template and player when no sequence is assigned, but that's currently not possible
 		if (InSequence)
 		{
-			SequencePlayer->Initialize(InSequence, GetWorld(), PlaybackSettings);
+			SequencePlayer->Initialize(InSequence, GetLevel(), PlaybackSettings);
 		}
 	}
 }
@@ -197,7 +197,7 @@ void ALevelSequenceActor::InitializePlayer()
 			// Level sequence is already loaded. Initialize the player if it's not already initialized with this sequence
 			if (LevelSequenceAsset != SequencePlayer->GetSequence())
 			{
-				SequencePlayer->Initialize(LevelSequenceAsset, GetWorld(), PlaybackSettings);
+				SequencePlayer->Initialize(LevelSequenceAsset, GetLevel(), PlaybackSettings);
 			}
 		}
 		else if (!IsAsyncLoading())
@@ -205,7 +205,7 @@ void ALevelSequenceActor::InitializePlayer()
 			LevelSequenceAsset = LoadSequence();
 			if (LevelSequenceAsset != SequencePlayer->GetSequence())
 			{
-				SequencePlayer->Initialize(LevelSequenceAsset, GetWorld(), PlaybackSettings);
+				SequencePlayer->Initialize(LevelSequenceAsset, GetLevel(), PlaybackSettings);
 			}
 		}
 		else
@@ -222,7 +222,7 @@ void ALevelSequenceActor::OnSequenceLoaded(const FName& PackageName, UPackage* P
 		ULevelSequence* LevelSequenceAsset = GetSequence();
 		if (SequencePlayer->GetSequence() != LevelSequenceAsset)
 		{
-			SequencePlayer->Initialize(LevelSequenceAsset, GetWorld(), PlaybackSettings);
+			SequencePlayer->Initialize(LevelSequenceAsset, GetLevel(), PlaybackSettings);
 		}
 	}
 }
diff --git a/Engine/Source/Runtime/LevelSequence/Private/LevelSequenceBindingReference.cpp b/Engine/Source/Runtime/LevelSequence/Private/LevelSequenceBindingReference.cpp
--- a/Engine/Source/Runtime/LevelSequence/Private/LevelSequenceBindingReference.cpp
+++ b/Engine/Source/Runtime/LevelSequence/Private/LevelSequenceBindingReference.cpp
@@ -9,6 +9,7 @@
 #include "Engine/World.h"
 #include "Animation/AnimInstance.h"
 #include "Components/SkeletalMeshComponent.h"
+#include "Engine/LevelStreamingDynamic.h"
 
 FLevelSequenceBindingReference::FLevelSequenceBindingReference(UObject* InObject, UObject* InContext)
 {
@@ -41,11 +42,19 @@ FLevelSequenceBindingReference::FLevelSequenceBindingReference(UObject* InObject
 	}
 }
 
-UObject* FLevelSequenceBindingReference::Resolve(UObject* InContext) const
+UObject* FLevelSequenceBindingReference::Resolve(UObject* InContext, FName StreamedLevelAssetPath) const
 {
-	if (ExternalObjectPath.IsNull())
+	if (InContext->IsA<AActor>())
 	{
-		return FindObject<UObject>(InContext, *ObjectPath, false);
+		if (ExternalObjectPath.IsNull())
+		{
+			return FindObject<UObject>(InContext, *ObjectPath, false);
+		}
+	}
+	else if (InContext->IsA<ULevel>() && StreamedLevelAssetPath != NAME_None && ExternalObjectPath.GetAssetPathName() == StreamedLevelAssetPath)
+	{
+		// ExternalObjectPath.GetSubPathString() specifies the path from the package (so includes PersistentLevel.) so we must do a FindObject from its outer
+		return FindObject<UObject>(InContext->GetOuter(), *ExternalObjectPath.GetSubPathString());
 	}
 	else
 	{
@@ -56,7 +65,7 @@ UObject* FLevelSequenceBindingReference::Resolve(UObject* InContext) const
 		// properly resolve. This fixes Possessable bindings losing their references the first time you save a map.
 		TempPath.PreSavePath();
 
-#if WITH_EDITORONLY_DATA
+	#if WITH_EDITORONLY_DATA
 		int32 ContextPlayInEditorID = InContext ? InContext->GetOutermost()->PIEInstanceID : INDEX_NONE;
 
 		if (ContextPlayInEditorID != INDEX_NONE)
@@ -69,10 +78,12 @@ UObject* FLevelSequenceBindingReference::Resolve(UObject* InContext) const
 		{
 			TempPath.FixupForPIE();
 		}
-#endif
+	#endif
 
 		return TempPath.ResolveObject();
 	}
+
+	return nullptr;
 }
 
 void FLevelSequenceBindingReference::PostSerialize(const FArchive& Ar)
@@ -221,7 +232,7 @@ void FLevelSequenceBindingReferences::RemoveObjects(const FGuid& ObjectId, const
 
 	for (int32 ReferenceIndex = 0; ReferenceIndex < ReferenceArray->References.Num(); )
 	{
-		UObject* ResolvedObject = ReferenceArray->References[ReferenceIndex].Resolve(InContext);
+		UObject* ResolvedObject = ReferenceArray->References[ReferenceIndex].Resolve(InContext, NAME_None);
 
 		if (InObjects.Contains(ResolvedObject))
 		{
@@ -244,7 +255,7 @@ void FLevelSequenceBindingReferences::RemoveInvalidObjects(const FGuid& ObjectId
 
 	for (int32 ReferenceIndex = 0; ReferenceIndex < ReferenceArray->References.Num(); )
 	{
-		UObject* ResolvedObject = ReferenceArray->References[ReferenceIndex].Resolve(InContext);
+		UObject* ResolvedObject = ReferenceArray->References[ReferenceIndex].Resolve(InContext, NAME_None);
 
 		if (!ResolvedObject || ResolvedObject->IsPendingKill())
 		{
@@ -257,34 +268,29 @@ void FLevelSequenceBindingReferences::RemoveInvalidObjects(const FGuid& ObjectId
 	}
 }
 
-void FLevelSequenceBindingReferences::ResolveBinding(const FGuid& ObjectId, UObject* InContext, TArray<UObject*, TInlineAllocator<1>>& OutObjects) const
+void FLevelSequenceBindingReferences::ResolveBinding(const FGuid& ObjectId, UObject* InContext, FName StreamedLevelAssetPath, TArray<UObject*, TInlineAllocator<1>>& OutObjects) const
 {
-	const FLevelSequenceBindingReferenceArray* ReferenceArray = BindingIdToReferences.Find(ObjectId);
-	if (!ReferenceArray)
+	if (const FLevelSequenceBindingReferenceArray* ReferenceArray = BindingIdToReferences.Find(ObjectId))
 	{
-		// If the object ID exists in the AnimSequenceInstances set, then this binding relates to an anim instance on a skeletal mesh component
-		USkeletalMeshComponent* SkeletalMesh = Cast<USkeletalMeshComponent>(InContext);
-		if (SkeletalMesh && AnimSequenceInstances.Contains(ObjectId) && SkeletalMesh->GetAnimInstance())
+		for (const FLevelSequenceBindingReference& Reference : ReferenceArray->References)
 		{
-			OutObjects.Add(SkeletalMesh->GetAnimInstance());
+			UObject* ResolvedObject = Reference.Resolve(InContext, StreamedLevelAssetPath);
+			if (ResolvedObject && ResolvedObject->GetWorld())
+			{
+				OutObjects.Add(ResolvedObject);
+			}
 		}
-
-		return;
 	}
-
-	for (const FLevelSequenceBindingReference& Reference : ReferenceArray->References)
+	else if (USkeletalMeshComponent* SkeletalMeshComponent = Cast<USkeletalMeshComponent>(InContext))
 	{
-		UObject* ResolvedObject = Reference.Resolve(InContext);
-
-		// if the resolved object does not have a valid world (e.g. world is being torn down), dont resolve
-		if (ResolvedObject && ResolvedObject->GetWorld())
+		// If the object ID exists in the AnimSequenceInstances set, then this binding relates to an anim instance on a skeletal mesh component
+		if (SkeletalMeshComponent && AnimSequenceInstances.Contains(ObjectId) && SkeletalMeshComponent->GetAnimInstance())
 		{
-			OutObjects.Add(ResolvedObject);
+			OutObjects.Add(SkeletalMeshComponent->GetAnimInstance());
 		}
 	}
 }
 
-
 void FLevelSequenceBindingReferences::RemoveInvalidBindings(const TSet<FGuid>& ValidBindingIDs)
 {
 	for (auto It = BindingIdToReferences.CreateIterator(); It; ++It)
diff --git a/Engine/Source/Runtime/LevelSequence/Private/LevelSequencePlayer.cpp b/Engine/Source/Runtime/LevelSequence/Private/LevelSequencePlayer.cpp
--- a/Engine/Source/Runtime/LevelSequence/Private/LevelSequencePlayer.cpp
+++ b/Engine/Source/Runtime/LevelSequence/Private/LevelSequencePlayer.cpp
@@ -21,7 +21,7 @@
 #include "Sections/MovieSceneCinematicShotSection.h"
 #include "LevelSequenceActor.h"
 #include "Modules/ModuleManager.h"
-
+#include "LevelUtils.h"
 
 /* ULevelSequencePlayer structors
  *****************************************************************************/
@@ -72,13 +72,50 @@ ULevelSequencePlayer* ULevelSequencePlayer::CreateLevelSequencePlayer(UObject* W
 /* ULevelSequencePlayer implementation
  *****************************************************************************/
 
-void ULevelSequencePlayer::Initialize(ULevelSequence* InLevelSequence, UWorld* InWorld, const FMovieSceneSequencePlaybackSettings& Settings)
+void ULevelSequencePlayer::Initialize(ULevelSequence* InLevelSequence, ULevel* InLevel, const FMovieSceneSequencePlaybackSettings& Settings)
 {
-	World = InWorld;
+	// Never use the level to resolve bindings unless we're playing back within a streamed or instanced level
+	StreamedLevelAssetPath = NAME_None;
+
+	World = InLevel->OwningWorld;
+	Level = InLevel;
+
+	// Construct the path to the level asset that the streamed level relates to
+	ULevelStreaming* LevelStreaming = FLevelUtils::FindStreamingLevel(InLevel);
+	if (LevelStreaming)
+	{
+		// StreamedLevelPackage is a package name of the form /Game/Folder/MapName, not a full asset path
+		FString StreamedLevelPackage = (LevelStreaming->PackageNameToLoad == NAME_None ? LevelStreaming->GetWorldAssetPackageFName() : LevelStreaming->PackageNameToLoad).ToString();
+
+		int32 SlashPos = 0;
+		if (StreamedLevelPackage.FindLastChar('/', SlashPos) && SlashPos < StreamedLevelPackage.Len()-1)
+		{
+			// Construct the asset path by appending .MapName to the end for efficient comparison with FSoftObjectPath::GetAssetPathName
+			const TCHAR* Pair[] = { *StreamedLevelPackage, &StreamedLevelPackage[SlashPos+1] };
+			StreamedLevelAssetPath = *FString::Join(Pair, TEXT("."));
+		}
+	}
+
 	SpawnRegister = MakeShareable(new FLevelSequenceSpawnRegister);
 	UMovieSceneSequencePlayer::Initialize(InLevelSequence, Settings);
 }
 
+void ULevelSequencePlayer::ResolveBoundObjects(const FGuid& InBindingId, FMovieSceneSequenceID SequenceID, UMovieSceneSequence& InSequence, UObject* ResolutionContext, TArray<UObject*, TInlineAllocator<1>>& OutObjects) const
+{
+	bool bAllowDefault = PlaybackClient ? PlaybackClient->RetrieveBindingOverrides(InBindingId, SequenceID, OutObjects) : true;
+
+	if (bAllowDefault)
+	{
+		if (StreamedLevelAssetPath != NAME_None && ResolutionContext && ResolutionContext->IsA<UWorld>())
+		{
+			ResolutionContext = Level.Get();
+		}
+
+		// Passing through the streamed level asset path ensures that bindings within instance sub levels resolve correctly
+		CastChecked<ULevelSequence>(&InSequence)->LocateBoundObjects(InBindingId, ResolutionContext, StreamedLevelAssetPath, OutObjects);
+	}
+}
+
 bool ULevelSequencePlayer::CanPlay() const
 {
 	return World.IsValid();
diff --git a/Engine/Source/Runtime/LevelSequence/Public/LevelSequence.h b/Engine/Source/Runtime/LevelSequence/Public/LevelSequence.h
--- a/Engine/Source/Runtime/LevelSequence/Public/LevelSequence.h
+++ b/Engine/Source/Runtime/LevelSequence/Public/LevelSequence.h
@@ -61,6 +61,8 @@ public:
 	virtual void GetAssetRegistryTags(TArray<FAssetRegistryTag>& OutTags) const override;
 #endif
 
+	void LocateBoundObjects(const FGuid& ObjectId, UObject* Context, FName StreamedLevelAssetPath, TArray<UObject*, TInlineAllocator<1>>& OutObjects) const;
+
 #if WITH_EDITORONLY_DATA
 
 	/** A pointer to the director blueprint that generates this sequence's DirectorClass. */
diff --git a/Engine/Source/Runtime/LevelSequence/Public/LevelSequenceBindingReference.h b/Engine/Source/Runtime/LevelSequence/Public/LevelSequenceBindingReference.h
--- a/Engine/Source/Runtime/LevelSequence/Public/LevelSequenceBindingReference.h
+++ b/Engine/Source/Runtime/LevelSequence/Public/LevelSequenceBindingReference.h
@@ -35,10 +35,11 @@ struct FLevelSequenceBindingReference
 	/**
 	 * Resolve this reference within the specified context
 	 *
-	 * @param InContext		The context to resolve the binding within. Either a UWorld or an AActor where this binding relates to an actor component
+	 * @param InContext		The context to resolve the binding within. Either a UWorld, ULevel (when playing in an instanced level) or an AActor where this binding relates to an actor component
+	 * @oaram StreamedLevelAssetPath    The path to the streamed level asset that contains the level sequence actor playing back the sequence. 'None' for any non-instance-level setups.
 	 * @return The object (usually an Actor or an ActorComponent).
 	 */
-	LEVELSEQUENCE_API UObject* Resolve(UObject* InContext) const;
+	LEVELSEQUENCE_API UObject* Resolve(UObject* InContext, FName StreamedLevelAssetPath) const;
 
 	/** Handles ExternalObjectPath fixup */
 	void PostSerialize(const FArchive& Ar);
@@ -131,11 +132,12 @@ struct FLevelSequenceBindingReferences
 	/**
 	 * Resolve a binding for the specified ID using a given context
 	 *
-	 * @param ObjectId		The ID to associate the object with
-	 * @param InContext		A context in which InObject resides
-	 * @param OutObjects	Array to populate with resolved object bindings
+	 * @param ObjectId					The ID to associate the object with
+	 * @param InContext					A context in which InObject resides
+	 * @oaram StreamedLevelAssetPath    The path to the streamed level asset that contains the level sequence actor playing back the sequence. 'None' for any non-instance-level setups.
+	 * @param OutObjects				Array to populate with resolved object bindings
 	 */
-	void ResolveBinding(const FGuid& ObjectId, UObject* InContext, TArray<UObject*, TInlineAllocator<1>>& OutObjects) const;
+	void ResolveBinding(const FGuid& ObjectId, UObject* InContext, FName StreamedLevelAssetPath, TArray<UObject*, TInlineAllocator<1>>& OutObjects) const;
 
 	/**
 	 * Const accessor for the currently bound anim instance IDs
diff --git a/Engine/Source/Runtime/LevelSequence/Public/LevelSequencePlayer.h b/Engine/Source/Runtime/LevelSequence/Public/LevelSequencePlayer.h
--- a/Engine/Source/Runtime/LevelSequence/Public/LevelSequencePlayer.h
+++ b/Engine/Source/Runtime/LevelSequence/Public/LevelSequencePlayer.h
@@ -111,10 +111,10 @@ public:
 	 * Initialize the player.
 	 *
 	 * @param InLevelSequence The level sequence to play.
-	 * @param InWorld The world that the animation is played in.
+	 * @param InLevel The level that the animation is played in.
 	 * @param Settings The desired playback settings
 	 */
-	void Initialize(ULevelSequence* InLevelSequence, UWorld* InWorld, const FMovieSceneSequencePlaybackSettings& Settings);
+	void Initialize(ULevelSequence* InLevelSequence, ULevel* InLevel, const FMovieSceneSequencePlaybackSettings& Settings);
 
 public:
 
@@ -154,6 +154,7 @@ protected:
 	// IMovieScenePlayer interface
 	virtual void UpdateCameraCut(UObject* CameraObject, UObject* UnlockIfCameraObject, bool bJumpCut) override;
 	virtual void NotifyBindingUpdate(const FGuid& InGuid, FMovieSceneSequenceIDRef InSequenceID, TArrayView<TWeakObjectPtr<>> Objects) override;
+	virtual void ResolveBoundObjects(const FGuid& InBindingId, FMovieSceneSequenceID SequenceID, UMovieSceneSequence& InSequence, UObject* ResolutionContext, TArray<UObject*, TInlineAllocator<1>>& OutObjects) const override;
 
 	//~ UMovieSceneSequencePlayer interface
 	virtual bool CanPlay() const override;
@@ -180,6 +181,12 @@ private:
 	/** The world this player will spawn actors in, if needed */
 	TWeakObjectPtr<UWorld> World;
 
+	/** The world this player will spawn actors in, if needed */
+	TWeakObjectPtr<ULevel> Level;
+
+	/** The full asset path (/Game/Folder/MapName.MapName) of the streaming level this player resides within. Bindings to actors with the same FSoftObjectPath::GetAssetPathName are resolved within the cached level, rather than globally.. */
+	FName StreamedLevelAssetPath;
+
 	/** The last view target to reset to when updating camera cuts to null */
 	TWeakObjectPtr<AActor> LastViewTarget;
 
