Sequencer: Fix PlaybackClient getting overwritten by CDO values on exit of ALevelSequenceActor constructor

diff --git a/Engine/Source/Runtime/LevelSequence/Private/LevelSequenceActor.cpp b/Engine/Source/Runtime/LevelSequence/Private/LevelSequenceActor.cpp
--- a/Engine/Source/Runtime/LevelSequence/Private/LevelSequenceActor.cpp
+++ b/Engine/Source/Runtime/LevelSequence/Private/LevelSequenceActor.cpp
@@ -52,7 +52,6 @@ ALevelSequenceActor::ALevelSequenceActor(const FObjectInitializer& Init)
 
 	// SequencePlayer must be a default sub object for it to be replicated correctly
 	SequencePlayer = Init.CreateDefaultSubobject<ULevelSequencePlayer>(this, "AnimationPlayer");
-	SequencePlayer->SetPlaybackClient(this);
 
 	bOverrideInstanceData = false;
 
@@ -63,6 +62,15 @@ ALevelSequenceActor::ALevelSequenceActor(const FObjectInitializer& Init)
 	bReplicatePlayback = false;
 }
 
+void ALevelSequenceActor::PostInitProperties()
+{
+	Super::PostInitProperties();
+
+	// Have to initialize this here as any properties set on default subobjects inside the constructor
+	// Get stomped by the CDO's properties when the constructor exits.
+	SequencePlayer->SetPlaybackClient(this);
+}
+
 bool ALevelSequenceActor::RetrieveBindingOverrides(const FGuid& InBindingId, FMovieSceneSequenceID InSequenceID, TArray<UObject*, TInlineAllocator<1>>& OutObjects) const
 {
 	return BindingOverrides->LocateBoundObjects(InBindingId, InSequenceID, OutObjects);
diff --git a/Engine/Source/Runtime/LevelSequence/Public/LevelSequenceActor.h b/Engine/Source/Runtime/LevelSequence/Public/LevelSequenceActor.h
--- a/Engine/Source/Runtime/LevelSequence/Public/LevelSequenceActor.h
+++ b/Engine/Source/Runtime/LevelSequence/Public/LevelSequenceActor.h
@@ -230,6 +230,7 @@ protected:
 
 	//~ Begin UObject interface
 	virtual bool ReplicateSubobjects(UActorChannel* Channel, FOutBunch* Bunch, FReplicationFlags *RepFlags) override;
+	virtual void PostInitProperties() override;
 	virtual void PostLoad() override;
 	//~ End UObject interface
 
