Sequencer: Check for valid context before resolving

diff --git a/Engine/Source/Runtime/LevelSequence/Private/LevelSequenceBindingReference.cpp b/Engine/Source/Runtime/LevelSequence/Private/LevelSequenceBindingReference.cpp
--- a/Engine/Source/Runtime/LevelSequence/Private/LevelSequenceBindingReference.cpp
+++ b/Engine/Source/Runtime/LevelSequence/Private/LevelSequenceBindingReference.cpp
@@ -44,14 +44,14 @@ FLevelSequenceBindingReference::FLevelSequenceBindingReference(UObject* InObject
 
 UObject* FLevelSequenceBindingReference::Resolve(UObject* InContext, FName StreamedLevelAssetPath) const
 {
-	if (InContext->IsA<AActor>())
+	if (InContext && InContext->IsA<AActor>())
 	{
 		if (ExternalObjectPath.IsNull())
 		{
 			return FindObject<UObject>(InContext, *ObjectPath, false);
 		}
 	}
-	else if (InContext->IsA<ULevel>() && StreamedLevelAssetPath != NAME_None && ExternalObjectPath.GetAssetPathName() == StreamedLevelAssetPath)
+	else if (InContext && InContext->IsA<ULevel>() && StreamedLevelAssetPath != NAME_None && ExternalObjectPath.GetAssetPathName() == StreamedLevelAssetPath)
 	{
 		// ExternalObjectPath.GetSubPathString() specifies the path from the package (so includes PersistentLevel.) so we must do a FindObject from its outer
 		return FindObject<UObject>(InContext->GetOuter(), *ExternalObjectPath.GetSubPathString());
