PurgeLegacyBlueprints will no longer be called from within ForEachObjectWithOuter is it renames objects that reside in hash tables that are being iterated over which may lead to undefined behavior.

diff --git a/Engine/Source/Runtime/LevelSequence/Private/LevelSequence.cpp b/Engine/Source/Runtime/LevelSequence/Private/LevelSequence.cpp
--- a/Engine/Source/Runtime/LevelSequence/Private/LevelSequence.cpp
+++ b/Engine/Source/Runtime/LevelSequence/Private/LevelSequence.cpp
@@ -188,9 +188,16 @@ void ULevelSequence::PostLoad()
 
 	if (GetLinkerCustomVersion(FSequencerObjectVersion::GUID) < FSequencerObjectVersion::PurgeSpawnableBlueprints)
 	{
-		// Remove any old generated classes from the package that will have been left behind from when we used blueprints for spawnables
-		UPackage* Package = GetOutermost();
-		ForEachObjectWithOuter(MovieScene, [=](UObject* InObject) { PurgeLegacyBlueprints(InObject, Package); }, false);
+		// Remove any old generated classes from the package that will have been left behind from when we used blueprints for spawnables		
+		{
+			UPackage* Package = GetOutermost();
+			TArray<UObject*> PackageSubobjects;
+			GetObjectsWithOuter(Package, PackageSubobjects, false);
+			for (UObject* ObjectInPackage : PackageSubobjects)
+			{
+				PurgeLegacyBlueprints(ObjectInPackage, Package);
+			}
+		}
 
 		// Remove any invalid object bindings
 		TSet<FGuid> ValidObjectBindings;
