Sequencer: Fixed some cases where the director class could become detatched from its blueprint (for example, on duplication of a level sequence

diff --git a/Engine/Plugins/MovieScene/LevelSequenceEditor/Source/LevelSequenceEditor/Private/LevelSequenceEditorModule.cpp b/Engine/Plugins/MovieScene/LevelSequenceEditor/Source/LevelSequenceEditor/Private/LevelSequenceEditorModule.cpp
--- a/Engine/Plugins/MovieScene/LevelSequenceEditor/Source/LevelSequenceEditor/Private/LevelSequenceEditorModule.cpp
+++ b/Engine/Plugins/MovieScene/LevelSequenceEditor/Source/LevelSequenceEditor/Private/LevelSequenceEditorModule.cpp
@@ -76,7 +76,7 @@ public:
 		{
 			virtual UBlueprint* RetrieveBlueprint(UObject* InObject) const override
 			{
-				return Cast<UBlueprint>(CastChecked<ULevelSequence>(InObject)->DirectorBlueprint);
+				return CastChecked<ULevelSequence>(InObject)->GetDirectorBlueprint();
 			}
 
 			virtual bool AssetContainsBlueprint(const FAssetData& InAssetData) const
diff --git a/Engine/Plugins/MovieScene/LevelSequenceEditor/Source/LevelSequenceEditor/Private/Misc/MovieSceneSequenceEditor_LevelSequence.h b/Engine/Plugins/MovieScene/LevelSequenceEditor/Source/LevelSequenceEditor/Private/Misc/MovieSceneSequenceEditor_LevelSequence.h
--- a/Engine/Plugins/MovieScene/LevelSequenceEditor/Source/LevelSequenceEditor/Private/Misc/MovieSceneSequenceEditor_LevelSequence.h
+++ b/Engine/Plugins/MovieScene/LevelSequenceEditor/Source/LevelSequenceEditor/Private/Misc/MovieSceneSequenceEditor_LevelSequence.h
@@ -19,7 +19,7 @@ struct FMovieSceneSequenceEditor_LevelSequence : FMovieSceneSequenceEditor
 	virtual UBlueprint* GetBlueprintForSequence(UMovieSceneSequence* InSequence) const override
 	{
 		ULevelSequence* LevelSequence = CastChecked<ULevelSequence>(InSequence);
-		return CastChecked<UBlueprint>(LevelSequence->DirectorBlueprint, ECastCheckedType::NullAllowed);
+		return LevelSequence->GetDirectorBlueprint();
 	}
 
 	virtual UBlueprint* CreateBlueprintForSequence(UMovieSceneSequence* InSequence) const override
@@ -35,8 +35,7 @@ struct FMovieSceneSequenceEditor_LevelSequence : FMovieSceneSequenceEditor
 		FName BlueprintName = "SequenceDirector";
 		Blueprint = FKismetEditorUtilities::CreateBlueprint(ULevelSequenceDirector::StaticClass(), InSequence, BlueprintName, BPTYPE_Normal, UBlueprint::StaticClass(), UBlueprintGeneratedClass::StaticClass());
 
-		LevelSequence->DirectorBlueprint = Blueprint;
-		LevelSequence->DirectorClass = CastChecked<UBlueprintGeneratedClass>(Blueprint->GeneratedClass);
+		LevelSequence->SetDirectorBlueprint(Blueprint);
 		return Blueprint;
 	}
 
diff --git a/Engine/Source/Runtime/LevelSequence/Private/LevelSequence.cpp b/Engine/Source/Runtime/LevelSequence/Private/LevelSequence.cpp
--- a/Engine/Source/Runtime/LevelSequence/Private/LevelSequence.cpp
+++ b/Engine/Source/Runtime/LevelSequence/Private/LevelSequence.cpp
@@ -163,11 +163,38 @@ void PurgeLegacyBlueprints(UObject* InObject, UPackage* Package)
 }
 #endif
 
+void ULevelSequence::PostDuplicate(bool bDuplicateForPIE)
+{
+	Super::PostDuplicate(bDuplicateForPIE);
+
+#if WITH_EDITORONLY_DATA
+	if (DirectorBlueprint)
+	{
+		DirectorClass = DirectorBlueprint->GeneratedClass.Get();
+
+		// Remove the binding for the director blueprint recompilation and re-add it to be sure there is only one entry in the list
+		DirectorBlueprint->OnCompiled().RemoveAll(this);
+		DirectorBlueprint->OnCompiled().AddUObject(this, &ULevelSequence::OnDirectorRecompiled);
+	}
+	else
+	{
+		DirectorClass = nullptr;
+	}
+#endif
+}
+
 void ULevelSequence::PostLoad()
 {
 	Super::PostLoad();
 
 #if WITH_EDITOR
+	if (DirectorBlueprint)
+	{
+		// Remove the binding for the director blueprint recompilation and re-add it to be sure there is only one entry in the list
+		DirectorBlueprint->OnCompiled().RemoveAll(this);
+		DirectorBlueprint->OnCompiled().AddUObject(this, &ULevelSequence::OnDirectorRecompiled);
+	}
+
 	TSet<FGuid> InvalidSpawnables;
 
 	for (int32 Index = 0; Index < MovieScene->GetSpawnableCount(); ++Index)
@@ -343,6 +370,37 @@ void ULevelSequence::UnbindInvalidObjects(const FGuid& ObjectId, UObject* InCont
 
 #if WITH_EDITOR
 
+UBlueprint* ULevelSequence::GetDirectorBlueprint() const
+{
+	return DirectorBlueprint;
+}
+
+void ULevelSequence::SetDirectorBlueprint(UBlueprint* NewDirectorBlueprint)
+{
+	if (DirectorBlueprint)
+	{
+		DirectorBlueprint->OnCompiled().RemoveAll(this);
+	}
+
+	DirectorBlueprint = NewDirectorBlueprint;
+
+	if (DirectorBlueprint)
+	{
+		DirectorClass = NewDirectorBlueprint->GeneratedClass.Get();
+		DirectorBlueprint->OnCompiled().AddUObject(this, &ULevelSequence::OnDirectorRecompiled);
+	}
+	else
+	{
+		DirectorClass = nullptr;
+	}
+}
+
+void ULevelSequence::OnDirectorRecompiled(UBlueprint* InCompiledBlueprint)
+{
+	ensure(InCompiledBlueprint == DirectorBlueprint);
+	DirectorClass = DirectorBlueprint->GeneratedClass.Get();
+}
+
 FGuid ULevelSequence::FindOrAddBinding(UObject* InObject)
 {
 	UObject* PlaybackContext = InObject ? InObject->GetWorld() : nullptr;
diff --git a/Engine/Source/Runtime/LevelSequence/Public/LevelSequence.h b/Engine/Source/Runtime/LevelSequence/Public/LevelSequence.h
--- a/Engine/Source/Runtime/LevelSequence/Public/LevelSequence.h
+++ b/Engine/Source/Runtime/LevelSequence/Public/LevelSequence.h
@@ -10,6 +10,7 @@
 #include "LevelSequenceLegacyObjectReference.h"
 #include "LevelSequence.generated.h"
 
+class UBlueprint;
 class UMovieScene;
 class UBlueprintGeneratedClass;
 
@@ -57,31 +58,33 @@ public:
 	virtual void GetAssetRegistryTags(TArray<FAssetRegistryTag>& OutTags) const override;
 #endif
 	virtual void PostLoad() override;
+	virtual void PostDuplicate(bool bDuplicateForPIE) override;
 
-#if WITH_EDITORONLY_DATA
+#if WITH_EDITOR
 
-	/** A pointer to the director blueprint that generates this sequence's DirectorClass. */
-	UPROPERTY()
-	UObject* DirectorBlueprint;
+public:
 
-#endif
+	/**
+	 * Assign a new director blueprint to this level sequence. The specified blueprint *must* be contained within this object.?	 */
+	void SetDirectorBlueprint(UBlueprint* NewDirectorBlueprint);
 
 	/**
-	 * The class that is used to spawn this level sequence's director instance.
-	 * Director instances are allocated on-demand one per sequence during evaluation and are used by event tracks for triggering events.
+	 * Retrieve the currently assigned director blueprint for this level sequence
 	 */
-	UPROPERTY()
-	UBlueprintGeneratedClass* DirectorClass;
+	UBlueprint* GetDirectorBlueprint() const;
 
 protected:
 
-#if WITH_EDITOR
-
 	virtual FGuid CreatePossessable(UObject* ObjectToPossess) override;
 	virtual FGuid CreateSpawnable(UObject* ObjectToSpawn) override;
 
 	FGuid FindOrAddBinding(UObject* ObjectToPossess);
 
+	/**
+	 * Invoked when this level sequence's director blueprint has been recompiled
+	 */
+	void OnDirectorRecompiled(UBlueprint*);
+
 #endif // WITH_EDITOR
 
 protected:
@@ -97,4 +100,19 @@ protected:
 	/** Deprecated property housing old possessed object bindings */
 	UPROPERTY()
 	TMap<FString, FLevelSequenceObject> PossessedObjects_DEPRECATED;
+
+#if WITH_EDITORONLY_DATA
+
+	/** A pointer to the director blueprint that generates this sequence's DirectorClass. */
+	UPROPERTY()
+	UBlueprint* DirectorBlueprint;
+
+#endif
+
+	/**
+	 * The class that is used to spawn this level sequence's director instance.
+	 * Director instances are allocated on-demand one per sequence during evaluation and are used by event tracks for triggering events.
+	 */
+	UPROPERTY()
+	UClass* DirectorClass;
 };
