Sequencer: Added an accessor for LevelSequenceActor::SequencePlayer that only returns the sequence player if it is valid for use

diff --git a/Engine/Source/Runtime/LevelSequence/Private/LevelSequenceActor.cpp b/Engine/Source/Runtime/LevelSequence/Private/LevelSequenceActor.cpp
--- a/Engine/Source/Runtime/LevelSequence/Private/LevelSequenceActor.cpp
+++ b/Engine/Source/Runtime/LevelSequence/Private/LevelSequenceActor.cpp
@@ -81,6 +81,11 @@ UObject* ALevelSequenceActor::GetInstanceData() const
 	return bOverrideInstanceData ? DefaultInstanceData : nullptr;
 }
 
+ULevelSequencePlayer* ALevelSequenceActor::GetSequencePlayer() const
+{
+	return SequencePlayer && SequencePlayer->GetSequence() ? SequencePlayer : nullptr;
+}
+
 void ALevelSequenceActor::SetReplicatePlayback(bool bInReplicatePlayback)
 {
 	bReplicatePlayback = bInReplicatePlayback;
diff --git a/Engine/Source/Runtime/LevelSequence/Public/LevelSequenceActor.h b/Engine/Source/Runtime/LevelSequence/Public/LevelSequenceActor.h
--- a/Engine/Source/Runtime/LevelSequence/Public/LevelSequenceActor.h
+++ b/Engine/Source/Runtime/LevelSequence/Public/LevelSequenceActor.h
@@ -77,7 +77,7 @@ public:
 	UPROPERTY(EditAnywhere, BlueprintReadOnly, Category="Playback", meta=(ShowOnlyInnerProperties))
 	FMovieSceneSequencePlaybackSettings PlaybackSettings;
 
-	UPROPERTY(Instanced, transient, replicated, BlueprintReadOnly, Category="Playback", meta=(ExposeFunctionCategories="Game|Cinematic"))
+	UPROPERTY(Instanced, transient, replicated, BlueprintReadOnly, BlueprintGetter=GetSequencePlayer, Category="Playback", meta=(ExposeFunctionCategories="Game|Cinematic"))
 	ULevelSequencePlayer* SequencePlayer;
 
 	UPROPERTY(EditAnywhere, BlueprintReadOnly, Category="General", meta=(AllowedClasses="LevelSequence"))
@@ -151,6 +151,12 @@ public:
 	UFUNCTION(BlueprintSetter)
 	void SetReplicatePlayback(bool ReplicatePlayback);
 
+	/**
+	 * Access this actor's sequence player, or None if it is not yet initialized
+	 */
+	UFUNCTION(BlueprintGetter)
+	ULevelSequencePlayer* GetSequencePlayer() const;
+
 	/** Refresh this actor's burn in */
 	void RefreshBurnIn();
 
