Sequencer: Added ensure if sequence director is not found.

diff --git a/Engine/Source/Runtime/LevelSequence/Private/LevelSequence.cpp b/Engine/Source/Runtime/LevelSequence/Private/LevelSequence.cpp
--- a/Engine/Source/Runtime/LevelSequence/Private/LevelSequence.cpp
+++ b/Engine/Source/Runtime/LevelSequence/Private/LevelSequence.cpp
@@ -211,6 +211,15 @@ void ULevelSequence::PostLoad()
 	Super::PostLoad();
 
 #if WITH_EDITOR
+	if (!DirectorBlueprint)
+	{
+		UBlueprint* PhantomDirector = FindObject<UBlueprint>(this, TEXT("SequenceDirector"));
+		if (!ensureMsgf(!PhantomDirector, TEXT("Phantom sequence director found in sequence '%s' which has a nullptr DirectorBlueprint. Re-assigning to prevent future crash."), *GetName()))
+		{
+			DirectorBlueprint = PhantomDirector;
+		}
+	}
+
 	if (DirectorBlueprint)
 	{
 		// Remove the binding for the director blueprint recompilation and re-add it to be sure there is only one entry in the list
