Sequencer: Remove WITH_EDITOR for development only functions

diff --git a/Engine/Source/Runtime/LevelSequence/Public/LevelSequence.h b/Engine/Source/Runtime/LevelSequence/Public/LevelSequence.h
--- a/Engine/Source/Runtime/LevelSequence/Public/LevelSequence.h
+++ b/Engine/Source/Runtime/LevelSequence/Public/LevelSequence.h
@@ -121,7 +121,6 @@ protected:
 	UPROPERTY()
 	UClass* DirectorClass;
 
-#if WITH_EDITOR
 public:
 	/**
 	* Find meta-data of a particular type for this level sequence instance.
@@ -184,7 +183,6 @@ public:
 	{
 		MetaDataObjects.RemoveAll([InClass](UObject* In) { return In && In->GetClass() == InClass; });
 	}
-#endif // WITH_EDITOR
 
 #if WITH_EDITORONLY_DATA
 
