Sequencer: Remove WITH_EDITOR for development only functions

diff --git a/Engine/Source/Runtime/LevelSequence/Public/LevelSequence.h b/Engine/Source/Runtime/LevelSequence/Public/LevelSequence.h
--- a/Engine/Source/Runtime/LevelSequence/Public/LevelSequence.h
+++ b/Engine/Source/Runtime/LevelSequence/Public/LevelSequence.h
@@ -130,8 +130,11 @@ public:
 	UFUNCTION(BlueprintPure, Category = "Level Sequence", meta=(DevelopmentOnly))
 	UObject* FindMetaDataByClass(TSubclassOf<UObject> InClass) const
 	{
+#if WITH_EDITORONLY_DATA
 		UObject* const* Found = MetaDataObjects.FindByPredicate([InClass](UObject* In) { return In && In->GetClass() == InClass; });
 		return Found ? CastChecked<UObject>(*Found) : nullptr;
+#endif
+		return nullptr;
 	}
 
 	/**
@@ -142,13 +145,14 @@ public:
 	UFUNCTION(BlueprintCallable, Category = "Level Sequence", meta=(DevelopmentOnly))
 	UObject* FindOrAddMetaDataByClass(TSubclassOf<UObject> InClass)
 	{
+#if WITH_EDITORONLY_DATA
 		UObject* Found = FindMetaDataByClass(InClass);
 		if (!Found)
 		{
 			Found = NewObject<UObject>(this, InClass);
 			MetaDataObjects.Add(Found);
 		}
-
+#endif
 		return Found;
 	}
 
@@ -161,6 +165,7 @@ public:
 	UFUNCTION(BlueprintCallable, Category = "Level Sequence", meta=(DevelopmentOnly))
 	UObject* CopyMetaData(UObject* InMetaData)
 	{
+#if WITH_EDITORONLY_DATA
 		if (!InMetaData)
 		{
 			return nullptr;
@@ -172,6 +177,8 @@ public:
 		MetaDataObjects.Add(NewMetaData);
 
 		return NewMetaData;
+#endif	
+		return nullptr;
 	}
 
 	/**
@@ -181,7 +188,9 @@ public:
 	UFUNCTION(BlueprintCallable, Category = "Level Sequence", meta=(DevelopmentOnly))
 	void RemoveMetaDataByClass(TSubclassOf<UObject> InClass)
 	{
+#if WITH_EDITORONLY_DATA
 		MetaDataObjects.RemoveAll([InClass](UObject* In) { return In && In->GetClass() == InClass; });
+#endif
 	}
 
 #if WITH_EDITORONLY_DATA
