Sequencer: Remove WITH_EDITOR for development only functions

diff --git a/Engine/Source/Runtime/LevelSequence/Public/LevelSequence.h b/Engine/Source/Runtime/LevelSequence/Public/LevelSequence.h
--- a/Engine/Source/Runtime/LevelSequence/Public/LevelSequence.h
+++ b/Engine/Source/Runtime/LevelSequence/Public/LevelSequence.h
@@ -152,8 +152,9 @@ public:
 			Found = NewObject<UObject>(this, InClass);
 			MetaDataObjects.Add(Found);
 		}
-#endif
 		return Found;
+#endif
+		return nullptr;
 	}
 
 	/**
