Making use of SetGameCameraCutThisFrame in various places
#rnx

diff --git a/Engine/Source/Developer/FunctionalTesting/Private/ScreenshotFunctionalTest.cpp b/Engine/Source/Developer/FunctionalTesting/Private/ScreenshotFunctionalTest.cpp
--- a/Engine/Source/Developer/FunctionalTesting/Private/ScreenshotFunctionalTest.cpp
+++ b/Engine/Source/Developer/FunctionalTesting/Private/ScreenshotFunctionalTest.cpp
@@ -47,7 +47,7 @@ void AScreenshotFunctionalTest::PrepareTest()
 
 		if (PlayerController && PlayerController->PlayerCameraManager)
 		{
-			PlayerController->PlayerCameraManager->bGameCameraCutThisFrame = true;
+			PlayerController->PlayerCameraManager->SetGameCameraCutThisFrame();
 			if (ScreenshotCamera)
 			{
 				ScreenshotCamera->NotifyCameraCut();
diff --git a/Engine/Source/Editor/Sequencer/Private/Sequencer.cpp b/Engine/Source/Editor/Sequencer/Private/Sequencer.cpp
--- a/Engine/Source/Editor/Sequencer/Private/Sequencer.cpp
+++ b/Engine/Source/Editor/Sequencer/Private/Sequencer.cpp
@@ -2905,7 +2905,7 @@ void FSequencer::PossessPIEViewports(UObject* CameraObject, UObject* UnlockIfCam
 		{
 			if (PC->PlayerCameraManager)
 			{
-				PC->PlayerCameraManager->bGameCameraCutThisFrame = true;
+				PC->PlayerCameraManager->SetGameCameraCutThisFrame();
 			}
 
 			if (CameraComponent)
@@ -2948,7 +2948,7 @@ void FSequencer::PossessPIEViewports(UObject* CameraObject, UObject* UnlockIfCam
 	if (PC->PlayerCameraManager)
 	{
 		PC->PlayerCameraManager->bClientSimulatingViewTarget = (CameraActor != nullptr);
-		PC->PlayerCameraManager->bGameCameraCutThisFrame = true;
+		PC->PlayerCameraManager->SetGameCameraCutThisFrame();
 	}
 }
 
diff --git a/Engine/Source/Runtime/Engine/Private/Interpolation.cpp b/Engine/Source/Runtime/Engine/Private/Interpolation.cpp
--- a/Engine/Source/Runtime/Engine/Private/Interpolation.cpp
+++ b/Engine/Source/Runtime/Engine/Private/Interpolation.cpp
@@ -6774,7 +6774,7 @@ void UInterpTrackDirector::UpdateTrack(float NewPosition, UInterpTrackInst* TrIn
 					
 					if ( PC->PlayerCameraManager )
 					{
-						PC->PlayerCameraManager->bGameCameraCutThisFrame = true;
+						PC->PlayerCameraManager->SetGameCameraCutThisFrame();
 					}
 					
 					DirInst->OldViewTarget = BackupViewTarget;
diff --git a/Engine/Source/Runtime/LevelSequence/Private/LevelSequencePlayer.cpp b/Engine/Source/Runtime/LevelSequence/Private/LevelSequencePlayer.cpp
--- a/Engine/Source/Runtime/LevelSequence/Private/LevelSequencePlayer.cpp
+++ b/Engine/Source/Runtime/LevelSequence/Private/LevelSequencePlayer.cpp
@@ -237,7 +237,7 @@ void ULevelSequencePlayer::UpdateCameraCut(UObject* CameraObject, UObject* Unloc
 		{
 			if (PC->PlayerCameraManager)
 			{
-				PC->PlayerCameraManager->bGameCameraCutThisFrame = true;
+				PC->PlayerCameraManager->SetGameCameraCutThisFrame();
 			}
 
 			if (CameraComponent)
@@ -288,7 +288,7 @@ void ULevelSequencePlayer::UpdateCameraCut(UObject* CameraObject, UObject* Unloc
 	if (PC->PlayerCameraManager)
 	{
 		PC->PlayerCameraManager->bClientSimulatingViewTarget = (CameraActor != nullptr);
-		PC->PlayerCameraManager->bGameCameraCutThisFrame = true;
+		PC->PlayerCameraManager->SetGameCameraCutThisFrame();
 	}
 
 	if (OnCameraCut.IsBound())
