Sequencer: Added ShowBurnin and HideBurnin functions

diff --git a/Engine/Source/Runtime/LevelSequence/Private/LevelSequenceActor.cpp b/Engine/Source/Runtime/LevelSequence/Private/LevelSequenceActor.cpp
--- a/Engine/Source/Runtime/LevelSequence/Private/LevelSequenceActor.cpp
+++ b/Engine/Source/Runtime/LevelSequence/Private/LevelSequenceActor.cpp
@@ -15,8 +15,11 @@
 	#include "SceneOutlinerFilters.h"
 #endif
 
+DEFINE_LOG_CATEGORY_STATIC(LogLevelSequence, Log, All);
+
 ALevelSequenceActor::ALevelSequenceActor(const FObjectInitializer& Init)
 	: Super(Init)
+	, bShowBurnin(true)
 {
 	USceneComponent* SceneComponent = CreateDefaultSubobject<USceneComponent>(TEXT("SceneComp"));
 	RootComponent = SceneComponent;
@@ -240,6 +243,28 @@ void ALevelSequenceActor::OnSequenceLoaded(const FName& PackageName, UPackage* P
 	}
 }
 
+void ALevelSequenceActor::HideBurnin()
+{
+	bShowBurnin = false;
+	RefreshBurnIn();
+
+	if (!BurnInOptions)
+	{
+		UE_LOG(LogLevelSequence, Warning, TEXT("Burnin is not enabled"));
+	}
+}
+
+void ALevelSequenceActor::ShowBurnin()
+{
+	bShowBurnin = true;
+	RefreshBurnIn();
+
+	if (!BurnInOptions || !BurnInOptions->bUseBurnIn)
+	{
+		UE_LOG(LogLevelSequence, Warning, TEXT("Burnin will not be visible because it is not enabled"));
+	}
+}
+
 void ALevelSequenceActor::RefreshBurnIn()
 {
 	if (BurnInInstance)
@@ -248,7 +273,7 @@ void ALevelSequenceActor::RefreshBurnIn()
 		BurnInInstance = nullptr;
 	}
 	
-	if (BurnInOptions && BurnInOptions->bUseBurnIn)
+	if (BurnInOptions && BurnInOptions->bUseBurnIn && bShowBurnin)
 	{
 		// Create the burn-in if necessary
 		UClass* Class = BurnInOptions->BurnInClass.TryLoadClass<ULevelSequenceBurnIn>();
diff --git a/Engine/Source/Runtime/LevelSequence/Public/LevelSequenceActor.h b/Engine/Source/Runtime/LevelSequence/Public/LevelSequenceActor.h
--- a/Engine/Source/Runtime/LevelSequence/Public/LevelSequenceActor.h
+++ b/Engine/Source/Runtime/LevelSequence/Public/LevelSequenceActor.h
@@ -157,6 +157,14 @@ public:
 	UFUNCTION(BlueprintGetter)
 	ULevelSequencePlayer* GetSequencePlayer() const;
 
+	/* Hide burnin */
+	UFUNCTION(BlueprintCallable, Category = "Game|Cinematic")
+	void HideBurnin();
+
+	/* Show burnin */
+	UFUNCTION(BlueprintCallable, Category = "Game|Cinematic")
+	void ShowBurnin();
+
 	/** Refresh this actor's burn in */
 	void RefreshBurnIn();
 
@@ -272,6 +280,9 @@ private:
 	/** Burn-in widget */
 	UPROPERTY()
 	ULevelSequenceBurnIn* BurnInInstance;
+
+	UPROPERTY()
+	bool bShowBurnin;
 };
 
 
