Sequencer: Fix log

diff --git a/Engine/Source/Runtime/LevelSequence/Private/LevelSequence.cpp b/Engine/Source/Runtime/LevelSequence/Private/LevelSequence.cpp
--- a/Engine/Source/Runtime/LevelSequence/Private/LevelSequence.cpp
+++ b/Engine/Source/Runtime/LevelSequence/Private/LevelSequence.cpp
@@ -24,8 +24,6 @@
 	#include "UObject/ObjectRedirector.h"
 #endif
 
-DEFINE_LOG_CATEGORY_STATIC(LogLevelSequence, Log, All);
-
 static TAutoConsoleVariable<int32> CVarDefaultEvaluationType(
 	TEXT("LevelSequence.DefaultEvaluationType"),
 	0,
diff --git a/Engine/Source/Runtime/LevelSequence/Private/LevelSequenceActor.cpp b/Engine/Source/Runtime/LevelSequence/Private/LevelSequenceActor.cpp
--- a/Engine/Source/Runtime/LevelSequence/Private/LevelSequenceActor.cpp
+++ b/Engine/Source/Runtime/LevelSequence/Private/LevelSequenceActor.cpp
@@ -8,6 +8,7 @@
 #include "DefaultLevelSequenceInstanceData.h"
 #include "Engine/ActorChannel.h"
 #include "Net/UnrealNetwork.h"
+#include "LevelSequenceModule.h"
 
 #if WITH_EDITOR
 	#include "PropertyCustomizationHelpers.h"
@@ -15,8 +16,6 @@
 	#include "SceneOutlinerFilters.h"
 #endif
 
-DEFINE_LOG_CATEGORY_STATIC(LogLevelSequence, Log, All);
-
 ALevelSequenceActor::ALevelSequenceActor(const FObjectInitializer& Init)
 	: Super(Init)
 	, bShowBurnin(true)
diff --git a/Engine/Source/Runtime/LevelSequence/Private/LevelSequenceModule.cpp b/Engine/Source/Runtime/LevelSequence/Private/LevelSequenceModule.cpp
--- a/Engine/Source/Runtime/LevelSequence/Private/LevelSequenceModule.cpp
+++ b/Engine/Source/Runtime/LevelSequence/Private/LevelSequenceModule.cpp
@@ -4,6 +4,8 @@
 #include "Modules/ModuleManager.h"
 #include "LevelSequenceActorSpawner.h"
 
+DEFINE_LOG_CATEGORY(LogLevelSequence);
+
 void FLevelSequenceModule::StartupModule()
 {
 	OnCreateMovieSceneObjectSpawnerDelegateHandle = RegisterObjectSpawner(FOnCreateMovieSceneObjectSpawner::CreateStatic(&FLevelSequenceActorSpawner::CreateObjectSpawner));
diff --git a/Engine/Source/Runtime/LevelSequence/Private/LevelSequenceModule.h b/Engine/Source/Runtime/LevelSequence/Private/LevelSequenceModule.h
--- a/Engine/Source/Runtime/LevelSequence/Private/LevelSequenceModule.h
+++ b/Engine/Source/Runtime/LevelSequence/Private/LevelSequenceModule.h
@@ -4,6 +4,8 @@
 
 #include "ILevelSequenceModule.h"
 
+DECLARE_LOG_CATEGORY_EXTERN(LogLevelSequence, Log, All);
+
 /**
  * Implements the LevelSequence module.
  */
