Sequencer: Make override instance data blueprint read write so it can be set in blueprints

diff --git a/Engine/Source/Runtime/LevelSequence/Public/LevelSequenceActor.h b/Engine/Source/Runtime/LevelSequence/Public/LevelSequenceActor.h
--- a/Engine/Source/Runtime/LevelSequence/Public/LevelSequenceActor.h
+++ b/Engine/Source/Runtime/LevelSequence/Public/LevelSequenceActor.h
@@ -97,7 +97,7 @@ public:
 	uint8 bAutoPlay_DEPRECATED : 1;
 
 	/** Enable specification of dynamic instance data to be supplied to the sequence during playback */
-	UPROPERTY(EditAnywhere, BlueprintReadOnly, Category="General")
+	UPROPERTY(EditAnywhere, BlueprintReadWrite, Category="General")
 	uint8 bOverrideInstanceData : 1;
 
 	/** If true, playback of this level sequence on the server will be synchronized across other clients */
