Sequencer: Tick sequencer separately.

This is needed so that sequencer actors can tick first, putting objects in their desired locations and then allowing other systems to respond. This fixes special cases where sequencer would tick animation again if it had ticked before sequencer.

diff --git a/Engine/Source/Runtime/Engine/Classes/Engine/World.h b/Engine/Source/Runtime/Engine/Classes/Engine/World.h
--- a/Engine/Source/Runtime/Engine/Classes/Engine/World.h
+++ b/Engine/Source/Runtime/Engine/Classes/Engine/World.h
@@ -843,6 +843,10 @@ class ENGINE_API UWorld final : public UObject, public FNetworkNotify
 	UPROPERTY(Transient)
 	TArray<UObject*>							PerModuleDataObjects;
 
+	// Level sequence actors to tick first
+	UPROPERTY(transient)
+	TArray<AActor*>								LevelSequenceActors;
+
 private:
 	/** Level collection. ULevels are referenced by FName (Package name) to avoid serialized references. Also contains offsets in world units */
 	UPROPERTY(Transient)
diff --git a/Engine/Source/Runtime/Engine/Private/LevelTick.cpp b/Engine/Source/Runtime/Engine/Private/LevelTick.cpp
--- a/Engine/Source/Runtime/Engine/Private/LevelTick.cpp
+++ b/Engine/Source/Runtime/Engine/Private/LevelTick.cpp
@@ -1535,6 +1535,15 @@ void UWorld::Tick( ELevelTick TickType, float DeltaSeconds )
 			SCOPE_CYCLE_COUNTER(STAT_TickTime);
 			FWorldDelegates::OnWorldPreActorTick.Broadcast(this, TickType, DeltaSeconds);
 		}
+
+		// Tick level sequence actors first
+		for (AActor* LevelSequenceActor : LevelSequenceActors)
+		{
+			if (LevelSequenceActor != nullptr)
+			{
+				LevelSequenceActor->Tick(DeltaSeconds);
+			}
+		}
 	}
 
 	for (int32 i = 0; i < LevelCollections.Num(); ++i)
diff --git a/Engine/Source/Runtime/LevelSequence/Private/LevelSequenceActor.cpp b/Engine/Source/Runtime/LevelSequence/Private/LevelSequenceActor.cpp
--- a/Engine/Source/Runtime/LevelSequence/Private/LevelSequenceActor.cpp
+++ b/Engine/Source/Runtime/LevelSequence/Private/LevelSequenceActor.cpp
@@ -57,7 +57,9 @@ ALevelSequenceActor::ALevelSequenceActor(const FObjectInitializer& Init)
 
 	bOverrideInstanceData = false;
 
-	PrimaryActorTick.bCanEverTick = true;
+	// The level sequence actor defaults to never ticking by the tick manager because it is ticked separately in LevelTick
+	//PrimaryActorTick.bCanEverTick = false;
+
 	bAutoPlay_DEPRECATED = false;
 
 	bReplicates = true;
@@ -120,6 +122,8 @@ void ALevelSequenceActor::PostInitializeComponents()
 
 void ALevelSequenceActor::BeginPlay()
 {
+	GetWorld()->LevelSequenceActors.Add(this);
+
 	Super::BeginPlay();
 
 	RefreshBurnIn();
@@ -130,6 +134,13 @@ void ALevelSequenceActor::BeginPlay()
 	}
 }
 
+void ALevelSequenceActor::EndPlay(const EEndPlayReason::Type EndPlayReason)
+{
+	GetWorld()->LevelSequenceActors.Remove(this);
+
+	Super::EndPlay(EndPlayReason);
+}
+
 void ALevelSequenceActor::Tick(float DeltaSeconds)
 {
 	Super::Tick(DeltaSeconds);
diff --git a/Engine/Source/Runtime/LevelSequence/Private/LevelSequencePlayer.cpp b/Engine/Source/Runtime/LevelSequence/Private/LevelSequencePlayer.cpp
--- a/Engine/Source/Runtime/LevelSequence/Private/LevelSequencePlayer.cpp
+++ b/Engine/Source/Runtime/LevelSequence/Private/LevelSequencePlayer.cpp
@@ -131,29 +131,6 @@ void ULevelSequencePlayer::OnStopped()
 {
 	EnableCinematicMode(false);
 
-	AActor* LevelSequenceActor = Cast<AActor>(GetOuter());
-	if (LevelSequenceActor == nullptr)
-	{
-		return;
-	}
-
-	for (FObjectKey WeakActor : PrerequisiteActors)
-	{
-		AActor* Actor = Cast<AActor>(WeakActor.ResolveObjectPtr());
-		if (Actor)
-		{
-			for (UActorComponent* Component : Actor->GetComponents())
-			{
-				if (Component)
-				{
-					Component->PrimaryComponentTick.RemovePrerequisite(LevelSequenceActor, LevelSequenceActor->PrimaryActorTick);
-				}
-			}
-
-			Actor->PrimaryActorTick.RemovePrerequisite(LevelSequenceActor, LevelSequenceActor->PrimaryActorTick);
-		}
-	}
-
 	if (World != nullptr && World->GetGameInstance() != nullptr)
 	{
 		APlayerController* PC = World->GetGameInstance()->GetFirstLocalPlayerController();
@@ -167,7 +144,6 @@ void ULevelSequencePlayer::OnStopped()
 		}
 	}
 
-	PrerequisiteActors.Reset();
 	LastViewTarget.Reset();
 }
 
@@ -297,38 +273,6 @@ void ULevelSequencePlayer::UpdateCameraCut(UObject* CameraObject, UObject* Unloc
 	}
 }
 
-void ULevelSequencePlayer::NotifyBindingUpdate(const FGuid& InGuid, FMovieSceneSequenceIDRef InSequenceID, TArrayView<TWeakObjectPtr<>> Objects)
-{
-	AActor* LevelSequenceActor = Cast<AActor>(GetOuter());
-	if (LevelSequenceActor == nullptr)
-	{
-		return;
-	}
-
-	for (TWeakObjectPtr<> WeakObject : Objects)
-	{
-		if (AActor* Actor = Cast<AActor>(WeakObject.Get()))
-		{
-			
-			if (Actor == LevelSequenceActor)
-			{
-				continue;
-			}
-
-			for (UActorComponent* Component : Actor->GetComponents())
-			{
-				if (Component)
-				{
-					Component->PrimaryComponentTick.AddPrerequisite(LevelSequenceActor, LevelSequenceActor->PrimaryActorTick);
-				}
-			}
-
-			Actor->PrimaryActorTick.AddPrerequisite(LevelSequenceActor, LevelSequenceActor->PrimaryActorTick);
-			PrerequisiteActors.Add(Actor);
-		}
-	}
-}
-
 UObject* ULevelSequencePlayer::GetPlaybackContext() const
 {
 	return World.Get();
diff --git a/Engine/Source/Runtime/LevelSequence/Public/LevelSequenceActor.h b/Engine/Source/Runtime/LevelSequence/Public/LevelSequenceActor.h
--- a/Engine/Source/Runtime/LevelSequence/Public/LevelSequenceActor.h
+++ b/Engine/Source/Runtime/LevelSequence/Public/LevelSequenceActor.h
@@ -252,6 +252,7 @@ protected:
 	virtual void Tick(float DeltaSeconds) override;
 	virtual void PostInitializeComponents() override;
 	virtual void BeginPlay() override;
+	virtual void EndPlay(const EEndPlayReason::Type EndPlayReason) override;
 	//~ End AActor interface
 
 public:
diff --git a/Engine/Source/Runtime/LevelSequence/Public/LevelSequencePlayer.h b/Engine/Source/Runtime/LevelSequence/Public/LevelSequencePlayer.h
--- a/Engine/Source/Runtime/LevelSequence/Public/LevelSequencePlayer.h
+++ b/Engine/Source/Runtime/LevelSequence/Public/LevelSequencePlayer.h
@@ -157,7 +157,6 @@ protected:
 
 	// IMovieScenePlayer interface
 	virtual void UpdateCameraCut(UObject* CameraObject, UObject* UnlockIfCameraObject, bool bJumpCut) override;
-	virtual void NotifyBindingUpdate(const FGuid& InGuid, FMovieSceneSequenceIDRef InSequenceID, TArrayView<TWeakObjectPtr<>> Objects) override;
 	virtual void ResolveBoundObjects(const FGuid& InBindingId, FMovieSceneSequenceID SequenceID, UMovieSceneSequence& InSequence, UObject* ResolutionContext, TArray<UObject*, TInlineAllocator<1>>& OutObjects) const override;
 
 	//~ UMovieSceneSequencePlayer interface
@@ -207,9 +206,6 @@ protected:
 
 	TWeakObjectPtr<UCameraComponent> CachedCameraComponent;
 
-	/** Set of actors that have been added as tick prerequisites to the parent actor */
-	TSet<FObjectKey> PrerequisiteActors;
-
 private:
 
 	TOptional<FLevelSequencePlayerSnapshot> PreviousSnapshot;
diff --git a/Engine/Source/Runtime/MovieScene/Private/MovieSceneSequencePlayer.cpp b/Engine/Source/Runtime/MovieScene/Private/MovieSceneSequencePlayer.cpp
--- a/Engine/Source/Runtime/MovieScene/Private/MovieSceneSequencePlayer.cpp
+++ b/Engine/Source/Runtime/MovieScene/Private/MovieSceneSequencePlayer.cpp
@@ -43,6 +43,7 @@ UMovieSceneSequencePlayer::UMovieSceneSequencePlayer(const FObjectInitializer& I
 	, StartTime(0)
 	, DurationFrames(0)
 	, CurrentNumLoops(0)
+	, LastTickGameTimeSeconds(-1.f)
 {
 	PlayPosition.Reset(FFrameTime(0));
 
@@ -666,9 +667,23 @@ void UMovieSceneSequencePlayer::Update(const float DeltaSeconds)
 		// Delta seconds has already been multiplied by MatineeTimeDilation at this point, so don't pass that through to Tick
 		float PlayRate = bReversePlayback ? -PlaybackSettings.PlayRate : PlaybackSettings.PlayRate;
 
-		TimeController->Tick(DeltaSeconds, PlayRate);
-
 		UWorld* World = GetPlaybackWorld();
+		float CurrentWorldTime = 0.f;
+		if (World)
+		{
+			CurrentWorldTime = World->GetTimeSeconds();
+		}
+
+		float DeltaTimeForFunction = DeltaSeconds;
+
+		if (LastTickGameTimeSeconds >= 0.f)
+		{
+			DeltaTimeForFunction = CurrentWorldTime - LastTickGameTimeSeconds;
+		}
+		LastTickGameTimeSeconds = CurrentWorldTime;
+
+		TimeController->Tick(DeltaTimeForFunction, PlayRate);
+
 		if (World)
 		{
 			PlayRate *= World->GetWorldSettings()->MatineeTimeDilation;
diff --git a/Engine/Source/Runtime/MovieScene/Public/MovieSceneSequencePlayer.h b/Engine/Source/Runtime/MovieScene/Public/MovieSceneSequencePlayer.h
--- a/Engine/Source/Runtime/MovieScene/Public/MovieSceneSequencePlayer.h
+++ b/Engine/Source/Runtime/MovieScene/Public/MovieSceneSequencePlayer.h
@@ -675,4 +675,10 @@ private:
 
 	/** The maximum tick rate prior to playing (used for overriding delta time during playback). */
 	TOptional<double> OldMaxTickRate;
+
+	/**
+	* The last world game time at which we were ticked. Game time used is dependent on bTickEvenWhenPaused
+	* Valid only if we've been ticked at least once since having a tick interval; otherwise set to -1.f
+	*/
+	float LastTickGameTimeSeconds;
 };
diff --git a/Engine/Source/Runtime/MovieSceneTracks/Private/Evaluation/MovieSceneSkeletalAnimationTemplate.cpp b/Engine/Source/Runtime/MovieSceneTracks/Private/Evaluation/MovieSceneSkeletalAnimationTemplate.cpp
--- a/Engine/Source/Runtime/MovieSceneTracks/Private/Evaluation/MovieSceneSkeletalAnimationTemplate.cpp
+++ b/Engine/Source/Runtime/MovieSceneTracks/Private/Evaluation/MovieSceneSkeletalAnimationTemplate.cpp
@@ -271,21 +271,6 @@ namespace MovieScene
 				}
 			}
 
-			// If the skeletal component has already ticked this frame because tick prerequisites weren't set up yet or a new binding was created, forcibly tick this component to update.
-			// This resolves first frame issues where the skeletal component ticks first, then the sequencer binding is resolved which sets up tick prerequisites
-			// for the next frame.
-			if (SkeletalMeshComponent->PoseTickedThisFrame() || (SequencerInstance && SequencerInstance != ExistingAnimInstance))
-			{
-				SkeletalMeshComponent->TickAnimation(0.f, false);
-
-				SkeletalMeshComponent->RefreshBoneTransforms();
-				SkeletalMeshComponent->RefreshSlaveComponents();
-				SkeletalMeshComponent->UpdateComponentToWorld();
-				SkeletalMeshComponent->FinalizeBoneTransform();
-				SkeletalMeshComponent->MarkRenderTransformDirty();
-				SkeletalMeshComponent->MarkRenderDynamicDataDirty();
-			}
-
 			Player.PreAnimatedState.SetCaptureEntity(FMovieSceneEvaluationKey(), EMovieSceneCompletionMode::KeepState);
 		}
 
