Sequencer: Default the snapshot's active shot to the current sequence in case there are no shots being rendered. This fixes an issue where {Slate} and {TakeNumber} don't work.

diff --git a/Engine/Source/Runtime/LevelSequence/Private/LevelSequencePlayer.cpp b/Engine/Source/Runtime/LevelSequence/Private/LevelSequencePlayer.cpp
--- a/Engine/Source/Runtime/LevelSequence/Private/LevelSequencePlayer.cpp
+++ b/Engine/Source/Runtime/LevelSequence/Private/LevelSequencePlayer.cpp
@@ -339,9 +339,14 @@ void ULevelSequencePlayer::TakeFrameSnapshot(FLevelSequencePlayerSnapshot& OutSn
 	OutSnapshot.CurrentShotLocalTime = FQualifiedFrameTime(CurrentPlayTime, PlayPosition.GetInputRate());
 	OutSnapshot.CameraComponent = CachedCameraComponent.IsValid() ? CachedCameraComponent.Get() : nullptr;
 	OutSnapshot.ShotID = MovieSceneSequenceID::Invalid;
-	OutSnapshot.ActiveShot = nullptr;
 
-	UMovieSceneCinematicShotTrack* ShotTrack = Sequence->GetMovieScene()->FindMasterTrack<UMovieSceneCinematicShotTrack>();
+	OutSnapshot.ActiveShot = Cast<ULevelSequence>(Sequence);
+
+	UMovieScene* MovieScene = Sequence->GetMovieScene();
+
+	OutSnapshot.SourceTimecode = MovieScene->TimecodeSource.Timecode.ToString();
+
+	UMovieSceneCinematicShotTrack* ShotTrack = MovieScene->FindMasterTrack<UMovieSceneCinematicShotTrack>();
 	if (ShotTrack)
 	{
 		UMovieSceneCinematicShotSection* ActiveShot = nullptr;
