Sequencer: Follow up CIS fix for 6088894

diff --git a/Engine/Source/Runtime/LevelSequence/Private/LevelSequencePlayer.cpp b/Engine/Source/Runtime/LevelSequence/Private/LevelSequencePlayer.cpp
--- a/Engine/Source/Runtime/LevelSequence/Private/LevelSequencePlayer.cpp
+++ b/Engine/Source/Runtime/LevelSequence/Private/LevelSequencePlayer.cpp
@@ -344,7 +344,9 @@ void ULevelSequencePlayer::TakeFrameSnapshot(FLevelSequencePlayerSnapshot& OutSn
 
 	UMovieScene* MovieScene = Sequence->GetMovieScene();
 
+#if WITH_EDITORONLY_DATA
 	OutSnapshot.SourceTimecode = MovieScene->TimecodeSource.Timecode.ToString();
+#endif
 
 	UMovieSceneCinematicShotTrack* ShotTrack = MovieScene->FindMasterTrack<UMovieSceneCinematicShotTrack>();
 	if (ShotTrack)
