Sequencer: Set replicates only if HasAuthority

diff --git a/Engine/Source/Runtime/LevelSequence/Private/LevelSequenceActor.cpp b/Engine/Source/Runtime/LevelSequence/Private/LevelSequenceActor.cpp
--- a/Engine/Source/Runtime/LevelSequence/Private/LevelSequenceActor.cpp
+++ b/Engine/Source/Runtime/LevelSequence/Private/LevelSequenceActor.cpp
@@ -116,7 +116,11 @@ void ALevelSequenceActor::PostInitializeComponents()
 {
 	Super::PostInitializeComponents();
 
-	SetReplicates(bReplicatePlayback);
+	if (HasAuthority())
+	{
+		SetReplicates(bReplicatePlayback);
+	}
+	
 	InitializePlayer();
 }
 
