Sequencer: Fix duplicate display names for FindBindingByTag

#rnx

diff --git a/Engine/Source/Runtime/LevelSequence/Public/LevelSequenceActor.h b/Engine/Source/Runtime/LevelSequence/Public/LevelSequenceActor.h
--- a/Engine/Source/Runtime/LevelSequence/Public/LevelSequenceActor.h
+++ b/Engine/Source/Runtime/LevelSequence/Public/LevelSequenceActor.h
@@ -240,7 +240,7 @@ public:
 	/**
 	 * Retrieve the first object binding that has been tagged with the specified name
 	 */
-	UFUNCTION(BlueprintCallable, Category = "Game|Cinematic|Bindings", DisplayName="FindBindingByTag")
+	UFUNCTION(BlueprintCallable, Category = "Game|Cinematic|Bindings", DisplayName="Find Binding by Tag")
 	FMovieSceneObjectBindingID FindNamedBinding(FName Tag) const;
 
 	/**
@@ -249,7 +249,7 @@ public:
 	 * @param Tag  The unique tag name to lookup bindings with. Object Bindings can be tagged within the sequence UI by RMB -> Tags... on the object binding in the tree.
 	 * @return An array containing all the bindings that are tagged with this name, potentially empty.
 	 */
-	UFUNCTION(BlueprintCallable, Category = "Game|Cinematic|Bindings", DisplayName="FindBindingByTag")
+	UFUNCTION(BlueprintCallable, Category = "Game|Cinematic|Bindings", DisplayName="Find Bindings by Tag")
 	const TArray<FMovieSceneObjectBindingID>& FindNamedBindings(FName Tag) const;
 
 protected:
