Sequencer: Wipe precompiled template data on director BP recompilation

This prevents us from tagging templates that reference old compiled data for export.

diff --git a/Engine/Source/Runtime/LevelSequence/Private/LevelSequence.cpp b/Engine/Source/Runtime/LevelSequence/Private/LevelSequence.cpp
--- a/Engine/Source/Runtime/LevelSequence/Private/LevelSequence.cpp
+++ b/Engine/Source/Runtime/LevelSequence/Private/LevelSequence.cpp
@@ -430,12 +430,18 @@ void ULevelSequence::SetDirectorBlueprint(UBlueprint* NewDirectorBlueprint)
 	{
 		DirectorClass = nullptr;
 	}
+
+	MarkAsChanged();
+	PrecompiledEvaluationTemplate = FMovieSceneEvaluationTemplate();
 }
 
 void ULevelSequence::OnDirectorRecompiled(UBlueprint* InCompiledBlueprint)
 {
 	ensure(InCompiledBlueprint == DirectorBlueprint);
 	DirectorClass = DirectorBlueprint->GeneratedClass.Get();
+
+	MarkAsChanged();
+	PrecompiledEvaluationTemplate = FMovieSceneEvaluationTemplate();
 }
 
 FGuid ULevelSequence::FindOrAddBinding(UObject* InObject)
