Add level sequence setting for handling the aspect ratio constraint during play.

The new setting defaults to overriding the constraint to XFOV to preserve previous behaviour. It lets the user choose what kind of constraint to use during the sequence, or whether to not touch the constraint at all.

diff --git a/Engine/Plugins/Developer/Concert/ConcertSync/ConcertSyncClient/Source/ConcertSyncClient/Private/ConcertClientSequencerManager.cpp b/Engine/Plugins/Developer/Concert/ConcertSync/ConcertSyncClient/Source/ConcertSyncClient/Private/ConcertClientSequencerManager.cpp
--- a/Engine/Plugins/Developer/Concert/ConcertSync/ConcertSyncClient/Source/ConcertSyncClient/Private/ConcertClientSequencerManager.cpp
+++ b/Engine/Plugins/Developer/Concert/ConcertSync/ConcertSyncClient/Source/ConcertSyncClient/Private/ConcertClientSequencerManager.cpp
@@ -472,7 +472,7 @@ void FConcertClientSequencerManager::ApplyEventToPlayers(const FConcertSequencer
 		if (Sequence && CurrentWorld)
 		{
 			Player = NewObject<ULevelSequencePlayer>((UObject*)GetTransientPackage(), FName("ConcertSequencePlayer"));
-			Player->Initialize(Sequence, CurrentWorld->PersistentLevel, FMovieSceneSequencePlaybackSettings());
+			Player->Initialize(Sequence, CurrentWorld->PersistentLevel, FMovieSceneSequencePlaybackSettings(), FLevelSequenceCameraSettings());
 		}
 		SequencePlayers.Add(*EventState.SequenceObjectPath, Player);
 	}
diff --git a/Engine/Plugins/Experimental/ControlRig/Source/ControlRigEditor/Private/Sequencer/ControlRigSequenceExporter.cpp b/Engine/Plugins/Experimental/ControlRig/Source/ControlRigEditor/Private/Sequencer/ControlRigSequenceExporter.cpp
--- a/Engine/Plugins/Experimental/ControlRig/Source/ControlRigEditor/Private/Sequencer/ControlRigSequenceExporter.cpp
+++ b/Engine/Plugins/Experimental/ControlRig/Source/ControlRigEditor/Private/Sequencer/ControlRigSequenceExporter.cpp
@@ -214,7 +214,7 @@ void Convert(UControlRigSequence* Sequence, UAnimSequence* AnimSequence, USkelet
 			LevelSequenceActor->LevelSequence = Sequence;
 			LevelSequenceActor->PlaybackSettings.bRestoreState = true;
 			LevelSequenceActor->SequencePlayer = NewObject<ULevelSequencePlayer>(LevelSequenceActor, "AnimationPlayer");
-			LevelSequenceActor->SequencePlayer->Initialize(Sequence, LevelSequenceActor->GetLevel(), LevelSequenceActor->PlaybackSettings);
+			LevelSequenceActor->SequencePlayer->Initialize(Sequence, LevelSequenceActor->GetLevel(), LevelSequenceActor->PlaybackSettings, LevelSequenceActor->CameraSettings);
 
 			// Now set up our animation sequence
 			AnimSequence->RecycleAnimSequence();
diff --git a/Engine/Plugins/MovieScene/SequencerScripting/Source/SequencerScriptingEditor/Private/SequencerTools.cpp b/Engine/Plugins/MovieScene/SequencerScripting/Source/SequencerScriptingEditor/Private/SequencerTools.cpp
--- a/Engine/Plugins/MovieScene/SequencerScripting/Source/SequencerScriptingEditor/Private/SequencerTools.cpp
+++ b/Engine/Plugins/MovieScene/SequencerScripting/Source/SequencerScriptingEditor/Private/SequencerTools.cpp
@@ -122,8 +122,9 @@ bool USequencerToolsFunctionLibrary::ExportFBX(UWorld* World, ULevelSequence* Se
 	INodeNameAdapter NodeNameAdapter;
 	ALevelSequenceActor* OutActor;
 	FMovieSceneSequencePlaybackSettings Settings;
+	FLevelSequenceCameraSettings CameraSettings;
 	ULevelSequencePlayer* Player = ULevelSequencePlayer::CreateLevelSequencePlayer(World, Sequence, Settings, OutActor);
-	Player->Initialize(Sequence, World->PersistentLevel, Settings);
+	Player->Initialize(Sequence, World->PersistentLevel, Settings, CameraSettings);
 	Player->State.AssignSequence(MovieSceneSequenceID::Root, *Sequence, *Player);
 	FMovieSceneSequenceIDRef Template = MovieSceneSequenceID::Root;
 	bool bDidExport = MovieSceneToolHelpers::ExportFBX(World, MovieScene, Player, Bindings, NodeNameAdapter, Template, InFBXFileName);
@@ -274,8 +275,9 @@ bool USequencerToolsFunctionLibrary::ImportFBX(UWorld* World, ULevelSequence* Se
 	const bool bMatchByNameOnly = ImportFBXSettings->bMatchByNameOnly;
 	ALevelSequenceActor* OutActor;
 	FMovieSceneSequencePlaybackSettings Settings;
+	FLevelSequenceCameraSettings CameraSettings;
 	ULevelSequencePlayer* Player = ULevelSequencePlayer::CreateLevelSequencePlayer(World, Sequence, Settings, OutActor);
-	Player->Initialize(Sequence, World->GetLevel(0), Settings);
+	Player->Initialize(Sequence, World->GetLevel(0), Settings, CameraSettings);
 	Player->State.AssignSequence(MovieSceneSequenceID::Root, *Sequence, *Player);
 
 	UnFbx::FFbxImporter* FbxImporter = UnFbx::FFbxImporter::GetInstance();
diff --git a/Engine/Source/Runtime/LevelSequence/Private/LevelSequenceActor.cpp b/Engine/Source/Runtime/LevelSequence/Private/LevelSequenceActor.cpp
--- a/Engine/Source/Runtime/LevelSequence/Private/LevelSequenceActor.cpp
+++ b/Engine/Source/Runtime/LevelSequence/Private/LevelSequenceActor.cpp
@@ -211,7 +211,7 @@ void ALevelSequenceActor::SetSequence(ULevelSequence* InSequence)
 		// cbb: should ideally null out the template and player when no sequence is assigned, but that's currently not possible
 		if (InSequence)
 		{
-			SequencePlayer->Initialize(InSequence, GetLevel(), PlaybackSettings);
+			SequencePlayer->Initialize(InSequence, GetLevel(), PlaybackSettings, CameraSettings);
 		}
 	}
 }
@@ -220,14 +220,14 @@ void ALevelSequenceActor::InitializePlayer()
 {
 	if (LevelSequence.IsValid() && GetWorld()->IsGameWorld())
 	{
-		// Attempt to reslove the asset without loading it
+		// Attempt to resolve the asset without loading it
 		ULevelSequence* LevelSequenceAsset = GetSequence();
 		if (LevelSequenceAsset)
 		{
 			// Level sequence is already loaded. Initialize the player if it's not already initialized with this sequence
 			if (LevelSequenceAsset != SequencePlayer->GetSequence())
 			{
-				SequencePlayer->Initialize(LevelSequenceAsset, GetLevel(), PlaybackSettings);
+				SequencePlayer->Initialize(LevelSequenceAsset, GetLevel(), PlaybackSettings, CameraSettings);
 			}
 		}
 		else if (!IsAsyncLoading())
@@ -235,7 +235,7 @@ void ALevelSequenceActor::InitializePlayer()
 			LevelSequenceAsset = LoadSequence();
 			if (LevelSequenceAsset != SequencePlayer->GetSequence())
 			{
-				SequencePlayer->Initialize(LevelSequenceAsset, GetLevel(), PlaybackSettings);
+				SequencePlayer->Initialize(LevelSequenceAsset, GetLevel(), PlaybackSettings, CameraSettings);
 			}
 		}
 		else
@@ -252,7 +252,7 @@ void ALevelSequenceActor::OnSequenceLoaded(const FName& PackageName, UPackage* P
 		ULevelSequence* LevelSequenceAsset = GetSequence();
 		if (SequencePlayer && SequencePlayer->GetSequence() != LevelSequenceAsset)
 		{
-			SequencePlayer->Initialize(LevelSequenceAsset, GetLevel(), PlaybackSettings);
+			SequencePlayer->Initialize(LevelSequenceAsset, GetLevel(), PlaybackSettings, CameraSettings);
 		}
 	}
 }
diff --git a/Engine/Source/Runtime/LevelSequence/Private/LevelSequencePlayer.cpp b/Engine/Source/Runtime/LevelSequence/Private/LevelSequencePlayer.cpp
--- a/Engine/Source/Runtime/LevelSequence/Private/LevelSequencePlayer.cpp
+++ b/Engine/Source/Runtime/LevelSequence/Private/LevelSequencePlayer.cpp
@@ -73,13 +73,14 @@ ULevelSequencePlayer* ULevelSequencePlayer::CreateLevelSequencePlayer(UObject* W
 /* ULevelSequencePlayer implementation
  *****************************************************************************/
 
-void ULevelSequencePlayer::Initialize(ULevelSequence* InLevelSequence, ULevel* InLevel, const FMovieSceneSequencePlaybackSettings& Settings)
+void ULevelSequencePlayer::Initialize(ULevelSequence* InLevelSequence, ULevel* InLevel, const FMovieSceneSequencePlaybackSettings& Settings, const FLevelSequenceCameraSettings& InCameraSettings)
 {
 	// Never use the level to resolve bindings unless we're playing back within a streamed or instanced level
 	StreamedLevelAssetPath = NAME_None;
 
 	World = InLevel->OwningWorld;
 	Level = InLevel;
+	CameraSettings = InCameraSettings;
 
 	// Construct the path to the level asset that the streamed level relates to
 	ULevelStreaming* LevelStreaming = FLevelUtils::FindStreamingLevel(InLevel);
@@ -184,16 +185,6 @@ void ULevelSequencePlayer::UpdateCameraCut(UObject* CameraObject, UObject* Unloc
 	// skip same view target
 	AActor* ViewTarget = PC->GetViewTarget();
 
-	// save the last view target so that it can be restored when the camera object is null
-	if (!LastViewTarget.IsValid())
-	{
-		LastViewTarget = ViewTarget;
-		if (PC->GetLocalPlayer())
-		{
-			LastAspectRatioAxisConstraint = PC->GetLocalPlayer()->AspectRatioAxisConstraint;
-		}
-	}
-
 	UCameraComponent* CameraComponent = MovieSceneHelpers::CameraComponentFromRuntimeObject(CameraObject);
 	if (CameraComponent && CameraComponent->GetOwner() != CameraObject)
 	{
@@ -237,9 +228,11 @@ void ULevelSequencePlayer::UpdateCameraCut(UObject* CameraObject, UObject* Unloc
 	AActor* CameraActor = Cast<AActor>(CameraObject);
 
 	// if the camera object is null, use the last view target so that it is restored to the state before the sequence takes control
+	bool bRestoreAspectRatioConstraint = false;
 	if (CameraActor == nullptr)
 	{
 		CameraActor = LastViewTarget.Get();
+		bRestoreAspectRatioConstraint = true;
 
 		// Skip if the last view target is the same as the current view target so that there's no additional camera cut
 		if (CameraActor == ViewTarget)
@@ -248,12 +241,39 @@ void ULevelSequencePlayer::UpdateCameraCut(UObject* CameraObject, UObject* Unloc
 		}
 	}
 
+	// Save the last view target/aspect ratio constraint/etc. so that it can all be restored when the camera object is null.
+	ULocalPlayer* LocalPlayer = PC->GetLocalPlayer();
+
+	if (!LastViewTarget.IsValid())
+	{
+		LastViewTarget = ViewTarget;
+	}
+	if (!LastAspectRatioAxisConstraint.IsSet())
+	{
+		if (LocalPlayer != nullptr)
+		{
+			LastAspectRatioAxisConstraint = LocalPlayer->AspectRatioAxisConstraint;
+		}
+	}
+
 	FViewTargetTransitionParams TransitionParams;
 	PC->SetViewTarget(CameraActor, TransitionParams);
 
-	if (PC->GetLocalPlayer())
+	// Set or restore the aspect ratio constraint if we were overriding it for this sequence.
+	if (LocalPlayer != nullptr && CameraSettings.bOverrideAspectRatioAxisConstraint)
 	{
-		PC->GetLocalPlayer()->AspectRatioAxisConstraint = EAspectRatioAxisConstraint::AspectRatio_MaintainXFOV;
+		if (bRestoreAspectRatioConstraint)
+		{
+			check(LastAspectRatioAxisConstraint.IsSet());
+			if (LastAspectRatioAxisConstraint.IsSet())
+			{
+				LocalPlayer->AspectRatioAxisConstraint = LastAspectRatioAxisConstraint.GetValue();
+			}
+		}
+		else
+		{
+			LocalPlayer->AspectRatioAxisConstraint = CameraSettings.AspectRatioAxisConstraint;
+		}
 	}
 
 	if (CameraComponent)
diff --git a/Engine/Source/Runtime/LevelSequence/Public/LevelSequenceActor.h b/Engine/Source/Runtime/LevelSequence/Public/LevelSequenceActor.h
--- a/Engine/Source/Runtime/LevelSequence/Public/LevelSequenceActor.h
+++ b/Engine/Source/Runtime/LevelSequence/Public/LevelSequenceActor.h
@@ -86,6 +86,9 @@ public:
 	UPROPERTY(EditAnywhere, BlueprintReadWrite, Category="General")
 	TArray<AActor*> AdditionalEventReceivers;
 
+	UPROPERTY(EditAnywhere, BlueprintReadWrite, Category="Cameras", meta=(ShowOnlyInnerProperties))
+	FLevelSequenceCameraSettings CameraSettings;
+
 	UPROPERTY(Instanced, BlueprintReadOnly, Category="General")
 	ULevelSequenceBurnInOptions* BurnInOptions;
 
diff --git a/Engine/Source/Runtime/LevelSequence/Public/LevelSequencePlayer.h b/Engine/Source/Runtime/LevelSequence/Public/LevelSequencePlayer.h
--- a/Engine/Source/Runtime/LevelSequence/Public/LevelSequencePlayer.h
+++ b/Engine/Source/Runtime/LevelSequence/Public/LevelSequencePlayer.h
@@ -92,6 +92,18 @@ struct FLevelSequencePlayerSnapshot
 	FMovieSceneSequenceID ShotID;
 };
 
+USTRUCT(BlueprintType)
+struct FLevelSequenceCameraSettings
+{
+	GENERATED_BODY()
+
+	UPROPERTY(EditAnywhere, BlueprintReadWrite, Category = "Aspect Ratio")
+	bool bOverrideAspectRatioAxisConstraint = true;
+
+	UPROPERTY(EditAnywhere, BlueprintReadWrite, Category = "Aspect Ratio", meta = (EditCondition = bOverrideAspectRatioAxisConstraint))
+	TEnumAsByte<enum EAspectRatioAxisConstraint> AspectRatioAxisConstraint = EAspectRatioAxisConstraint::AspectRatio_MaintainXFOV;
+};
+
 /**
  * ULevelSequencePlayer is used to actually "play" an level sequence asset at runtime.
  *
@@ -114,7 +126,10 @@ public:
 	 * @param InLevel The level that the animation is played in.
 	 * @param Settings The desired playback settings
 	 */
-	void Initialize(ULevelSequence* InLevelSequence, ULevel* InLevel, const FMovieSceneSequencePlaybackSettings& Settings);
+	void Initialize(ULevelSequence* InLevelSequence, ULevel* InLevel, const FMovieSceneSequencePlaybackSettings& Settings, const FLevelSequenceCameraSettings& InCameraSettings);
+
+	UE_DEPRECATED(4.23, "Added camera settings to Initialize.")
+	void Initialize(ULevelSequence* InLevelSequence, ULevel* InLevel, const FMovieSceneSequencePlaybackSettings& Settings) { Initialize(InLevelSequence, InLevel, Settings, FLevelSequenceCameraSettings()); }
 
 public:
 
@@ -191,11 +206,14 @@ private:
 	/** The full asset path (/Game/Folder/MapName.MapName) of the streaming level this player resides within. Bindings to actors with the same FSoftObjectPath::GetAssetPathName are resolved within the cached level, rather than globally.. */
 	FName StreamedLevelAssetPath;
 
+	/** The camera settings to use when playing the sequence */
+	FLevelSequenceCameraSettings CameraSettings;
+
 	/** The last view target to reset to when updating camera cuts to null */
 	TWeakObjectPtr<AActor> LastViewTarget;
 
 	/** The last aspect ratio axis constraint to reset to when the camera cut is null */
-	EAspectRatioAxisConstraint LastAspectRatioAxisConstraint;
+	TOptional<EAspectRatioAxisConstraint> LastAspectRatioAxisConstraint;
 
 protected:
 
