Sequencer: Correctly add PIE prefix to fix spawnable objects in streamed levels during PIE

diff --git a/Engine/Source/Runtime/LevelSequence/Private/LevelSequenceActorSpawner.cpp b/Engine/Source/Runtime/LevelSequence/Private/LevelSequenceActorSpawner.cpp
--- a/Engine/Source/Runtime/LevelSequence/Private/LevelSequenceActorSpawner.cpp
+++ b/Engine/Source/Runtime/LevelSequence/Private/LevelSequenceActorSpawner.cpp
@@ -22,20 +22,30 @@ UClass* FLevelSequenceActorSpawner::GetSupportedTemplateType() const
 	return AActor::StaticClass();
 }
 
-ULevelStreaming* GetLevelStreaming(const FName& DesiredLevelName, const TArray<ULevelStreaming*>& StreamingLevels)
+ULevelStreaming* GetLevelStreaming(const FName& DesiredLevelName, const UWorld* World)
 {
 	if (DesiredLevelName == NAME_None)
 	{
 		return nullptr;
 	}
 
+	const TArray<ULevelStreaming*>& StreamingLevels = World->GetStreamingLevels();
 	FString SafeLevelNameString = DesiredLevelName.ToString();
 	if (FPackageName::IsShortPackageName(SafeLevelNameString))
 	{
 		// Make sure MyMap1 and Map1 names do not resolve to a same streaming level
-		SafeLevelNameString.InsertAt('/', 0);
+		SafeLevelNameString.InsertAt(0, '/');
 	}
 
+#if WITH_EDITOR
+	FWorldContext* WorldContext = GEngine->GetWorldContextFromWorld(World);
+	if (WorldContext && WorldContext->PIEInstance != INDEX_NONE)
+	{
+		SafeLevelNameString = UWorld::ConvertToPIEPackageName(SafeLevelNameString, WorldContext->PIEInstance);
+	}
+#endif
+
+
 	for (ULevelStreaming* LevelStreaming : StreamingLevels)
 	{
 		if (LevelStreaming && LevelStreaming->GetWorldAssetPackageName().EndsWith(SafeLevelNameString, ESearchCase::IgnoreCase))
@@ -70,7 +80,7 @@ UObject* FLevelSequenceActorSpawner::SpawnObject(FMovieSceneSpawnable& Spawnable
 	FName DesiredLevelName = Spawnable.GetLevelName();
 	if (DesiredLevelName != NAME_None)
 	{
-		ULevelStreaming* LevelStreaming = GetLevelStreaming(DesiredLevelName, WorldContext->GetStreamingLevels());
+		ULevelStreaming* LevelStreaming = GetLevelStreaming(DesiredLevelName, WorldContext);
 		if (LevelStreaming && LevelStreaming->GetWorldAsset().IsValid())
 		{
 			WorldContext = LevelStreaming->GetWorldAsset().Get();
