Sequencer: ALevelSequenceActor now implements RewindForReplay which will stop the sequence when seeking starts in a replay. This solves issues where seeking before the sequence started would not restore the state as no new network packets are sent to update the state. Reset NetSyncProps of the sequence player when RewindForReplay() is called to prevent the sequence from playing at the wrong time

diff --git a/Engine/Source/Runtime/LevelSequence/Private/LevelSequenceActor.cpp b/Engine/Source/Runtime/LevelSequence/Private/LevelSequenceActor.cpp
--- a/Engine/Source/Runtime/LevelSequence/Private/LevelSequenceActor.cpp
+++ b/Engine/Source/Runtime/LevelSequence/Private/LevelSequenceActor.cpp
@@ -86,6 +86,14 @@ void ALevelSequenceActor::PostInitProperties()
 	SequencePlayer->SetPlaybackClient(this);
 }
 
+void ALevelSequenceActor::RewindForReplay()
+{
+	if (SequencePlayer)
+	{
+		SequencePlayer->RewindForReplay();
+	}
+}
+
 bool ALevelSequenceActor::RetrieveBindingOverrides(const FGuid& InBindingId, FMovieSceneSequenceID InSequenceID, TArray<UObject*, TInlineAllocator<1>>& OutObjects) const
 {
 	return BindingOverrides->LocateBoundObjects(InBindingId, InSequenceID, OutObjects);
diff --git a/Engine/Source/Runtime/LevelSequence/Private/LevelSequencePlayer.cpp b/Engine/Source/Runtime/LevelSequence/Private/LevelSequencePlayer.cpp
--- a/Engine/Source/Runtime/LevelSequence/Private/LevelSequencePlayer.cpp
+++ b/Engine/Source/Runtime/LevelSequence/Private/LevelSequencePlayer.cpp
@@ -457,3 +457,14 @@ void ULevelSequencePlayer::EnableCinematicMode(bool bEnable)
 	}
 }
 
+void ULevelSequencePlayer::RewindForReplay()
+{
+	// Stop the sequence when starting to seek through a replay. This restores our state to be unmodified
+	// in case the replay is seeking to before playback. If we're in the middle of playback after rewinding,
+	// the replay will feed the correct packets to synchronize our playback time and state.
+	Stop();
+
+	NetSyncProps.LastKnownPosition = FFrameTime(0);
+	NetSyncProps.LastKnownStatus = EMovieScenePlayerStatus::Stopped;
+	NetSyncProps.LastKnownNumLoops = 0;
+}
\ No newline at end of file
diff --git a/Engine/Source/Runtime/LevelSequence/Public/LevelSequenceActor.h b/Engine/Source/Runtime/LevelSequence/Public/LevelSequenceActor.h
--- a/Engine/Source/Runtime/LevelSequence/Public/LevelSequenceActor.h
+++ b/Engine/Source/Runtime/LevelSequence/Public/LevelSequenceActor.h
@@ -270,6 +270,7 @@ protected:
 	virtual void PostInitializeComponents() override;
 	virtual void BeginPlay() override;
 	virtual void EndPlay(const EEndPlayReason::Type EndPlayReason) override;
+	virtual void RewindForReplay() override;
 	//~ End AActor interface
 
 public:
diff --git a/Engine/Source/Runtime/LevelSequence/Public/LevelSequencePlayer.h b/Engine/Source/Runtime/LevelSequence/Public/LevelSequencePlayer.h
--- a/Engine/Source/Runtime/LevelSequence/Public/LevelSequencePlayer.h
+++ b/Engine/Source/Runtime/LevelSequence/Public/LevelSequencePlayer.h
@@ -168,6 +168,8 @@ public:
 	virtual UObject* GetPlaybackContext() const override;
 	virtual TArray<UObject*> GetEventContexts() const override;
 
+	void RewindForReplay();
+
 protected:
 
 	// IMovieScenePlayer interface
