Sequencer: Rename the director blueprint when the level sequence is renamed. This fixes a bug where if the level sequence is renamed, event tracks would stop working.

diff --git a/Engine/Source/Runtime/LevelSequence/Private/LevelSequence.cpp b/Engine/Source/Runtime/LevelSequence/Private/LevelSequence.cpp
--- a/Engine/Source/Runtime/LevelSequence/Private/LevelSequence.cpp
+++ b/Engine/Source/Runtime/LevelSequence/Private/LevelSequence.cpp
@@ -288,6 +288,18 @@ void ULevelSequence::PostLoad()
 #endif
 }
 
+bool ULevelSequence::Rename(const TCHAR* NewName, UObject* NewOuter, ERenameFlags Flags)
+{
+	bool bRetVal = Super::Rename(NewName, NewOuter, Flags);
+
+	if (DirectorBlueprint)
+	{
+		DirectorBlueprint->Rename(*DirectorBlueprint->GetName(), this, Flags);
+	}
+
+	return bRetVal;
+}
+
 void ULevelSequence::ConvertPersistentBindingsToDefault(UObject* FixupContext)
 {
 	if (PossessedObjects_DEPRECATED.Num() == 0)
diff --git a/Engine/Source/Runtime/LevelSequence/Public/LevelSequence.h b/Engine/Source/Runtime/LevelSequence/Public/LevelSequence.h
--- a/Engine/Source/Runtime/LevelSequence/Public/LevelSequence.h
+++ b/Engine/Source/Runtime/LevelSequence/Public/LevelSequence.h
@@ -56,6 +56,7 @@ public:
 	virtual bool CanAnimateObject(UObject& InObject) const override;
 	virtual UObject* CreateDirectorInstance(IMovieScenePlayer& Player) override;
 	virtual void PostLoad() override;
+	virtual bool Rename(const TCHAR* NewName = nullptr, UObject* NewOuter = nullptr, ERenameFlags Flags = REN_None) override;
 
 #if WITH_EDITOR
 	virtual void GetAssetRegistryTagMetadata(TMap<FName, FAssetRegistryTagMetadata>& OutMetadata) const override;
