Fix editor only
#rnx

diff --git a/Engine/Source/Runtime/LevelSequence/Private/LevelSequence.cpp b/Engine/Source/Runtime/LevelSequence/Private/LevelSequence.cpp
--- a/Engine/Source/Runtime/LevelSequence/Private/LevelSequence.cpp
+++ b/Engine/Source/Runtime/LevelSequence/Private/LevelSequence.cpp
@@ -292,10 +292,12 @@ bool ULevelSequence::Rename(const TCHAR* NewName, UObject* NewOuter, ERenameFlag
 {
 	bool bRetVal = Super::Rename(NewName, NewOuter, Flags);
 
+#if WITH_EDITOR
 	if (DirectorBlueprint)
 	{
 		DirectorBlueprint->Rename(*DirectorBlueprint->GetName(), this, Flags);
 	}
+#endif
 
 	return bRetVal;
 }
