Sequencer: Clear transient flags on spawnable template components. This fixes a bug where adding components to an empty actor would not save those components.

diff --git a/Engine/Source/Runtime/LevelSequence/Private/LevelSequence.cpp b/Engine/Source/Runtime/LevelSequence/Private/LevelSequence.cpp
--- a/Engine/Source/Runtime/LevelSequence/Private/LevelSequence.cpp
+++ b/Engine/Source/Runtime/LevelSequence/Private/LevelSequence.cpp
@@ -83,6 +83,16 @@ UObject* ULevelSequence::MakeSpawnableTemplateFromInstance(UObject& InSourceObje
 		Actor->DetachFromActor(FDetachmentTransformRules(FAttachmentTransformRules(EAttachmentRule::KeepRelative, false), false));
 	}
 
+	// The spawnable source object was created with RF_Transient. The object generated from that needs its 
+	// component flags cleared of RF_Transient so that the template object can be saved to the level sequence.
+	for (UActorComponent* Component : Actor->GetComponents())
+	{
+		if (Component)
+		{
+			Component->ClearFlags(RF_Transient);
+		}
+	}
+
 	return NewInstance;
 }
 
