Sequencer: Spawnables no longer defaults to spawning into the persistent level if the named level isn't found. Also, when a spawnable is added to sequencer, set the current level as the level to spawn into.

diff --git a/Engine/Source/Editor/Sequencer/Private/Sequencer.cpp b/Engine/Source/Editor/Sequencer/Private/Sequencer.cpp
--- a/Engine/Source/Editor/Sequencer/Private/Sequencer.cpp
+++ b/Engine/Source/Editor/Sequencer/Private/Sequencer.cpp
@@ -5017,6 +5017,16 @@ FGuid FSequencer::MakeNewSpawnable( UObject& Object, UActorFactory* ActorFactory
 	ESpawnOwnership SavedOwnership = Spawnable->GetSpawnOwnership();
 	Spawnable->SetSpawnOwnership(ESpawnOwnership::External);
 
+	// Spawn into the current level
+	FName LevelName;
+	UWorld* World = Cast<UWorld>(GetPlaybackContext());
+	if (World && World->GetCurrentLevel() && World->GetCurrentLevel() != World->PersistentLevel)
+	{
+		LevelName = FPackageName::GetShortFName(World->GetCurrentLevel()->GetOutermost()->GetFName());
+	}
+
+	Spawnable->SetLevelName(LevelName);
+
 	// Spawn the object so we can position it correctly, it's going to get spawned anyway since things default to spawned.
 	UObject* SpawnedObject = SpawnRegister->SpawnObject(NewGuid, *MovieScene, ActiveTemplateIDs.Top(), *this);
 
diff --git a/Engine/Source/Runtime/LevelSequence/Private/LevelSequenceActorSpawner.cpp b/Engine/Source/Runtime/LevelSequence/Private/LevelSequenceActorSpawner.cpp
--- a/Engine/Source/Runtime/LevelSequence/Private/LevelSequenceActorSpawner.cpp
+++ b/Engine/Source/Runtime/LevelSequence/Private/LevelSequenceActorSpawner.cpp
@@ -72,10 +72,6 @@ UObject* FLevelSequenceActorSpawner::SpawnObject(FMovieSceneSpawnable& Spawnable
 	// @todo sequencer actors: We need to make sure puppet objects aren't copied into PIE/SIE sessions!  They should be omitted from that duplication!
 
 	UWorld* WorldContext = Cast<UWorld>(Player.GetPlaybackContext());
-	if(WorldContext == nullptr)
-	{
-		WorldContext = GWorld;
-	}
 
 	FName DesiredLevelName = Spawnable.GetLevelName();
 	if (DesiredLevelName != NAME_None)
@@ -87,10 +83,27 @@ UObject* FLevelSequenceActorSpawner::SpawnObject(FMovieSceneSpawnable& Spawnable
 		}
 		else
 		{
-			UE_LOG(LogMovieScene, Warning, TEXT("Can't find sublevel '%s' to spawn '%s' into"), *DesiredLevelName.ToString(), *Spawnable.GetName());
+			// Avoid spamming output, warning only once per level
+			if (!ErrorLevels.Contains(DesiredLevelName))
+			{
+				UE_LOG(LogMovieScene, Warning, TEXT("Can't find sublevel '%s' to spawn '%s' into."), *DesiredLevelName.ToString(), *Spawnable.GetName());
+				ErrorLevels.Add(DesiredLevelName);
+			}
+			return nullptr;
 		}
 	}
 
+	if (WorldContext == nullptr)
+	{
+		if (!ErrorLevels.Contains(DesiredLevelName))
+		{
+			UE_LOG(LogMovieScene, Warning, TEXT("Can't find world to spawn '%s' into."), *Spawnable.GetName());
+			ErrorLevels.Add(DesiredLevelName);
+		}
+
+		return nullptr;
+	}
+
 	// Construct the object with the same name that we will set later on the actor to avoid renaming it inside SetActorLabel
 	FName SpawnName =
 #if WITH_EDITOR
diff --git a/Engine/Source/Runtime/LevelSequence/Public/LevelSequenceActorSpawner.h b/Engine/Source/Runtime/LevelSequence/Public/LevelSequenceActorSpawner.h
--- a/Engine/Source/Runtime/LevelSequence/Public/LevelSequenceActorSpawner.h
+++ b/Engine/Source/Runtime/LevelSequence/Public/LevelSequenceActorSpawner.h
@@ -14,4 +14,7 @@ public:
 	virtual UClass* GetSupportedTemplateType() const override;
 	virtual UObject* SpawnObject(FMovieSceneSpawnable& Spawnable, FMovieSceneSequenceIDRef TemplateID, IMovieScenePlayer& Player) override;
 	virtual void DestroySpawnedObject(UObject& Object) override;
+
+private:
+	TSet<FName> ErrorLevels;
 };
\ No newline at end of file
