Fix sequencer bindings when using external actors.

diff --git a/Engine/Source/Runtime/LevelSequence/Private/LevelSequenceBindingReference.cpp b/Engine/Source/Runtime/LevelSequence/Private/LevelSequenceBindingReference.cpp
--- a/Engine/Source/Runtime/LevelSequence/Private/LevelSequenceBindingReference.cpp
+++ b/Engine/Source/Runtime/LevelSequence/Private/LevelSequenceBindingReference.cpp
@@ -21,23 +21,22 @@ FLevelSequenceBindingReference::FLevelSequenceBindingReference(UObject* InObject
 	}
 	else
 	{
-		UPackage* ObjectPackage = InObject->GetOutermost();
-		if (!ensure(ObjectPackage))
-		{
-			return;
-		}
+		FString FullPath = InObject->GetPathName();
 
-		FString PackageName = ObjectPackage->GetName();
 #if WITH_EDITORONLY_DATA
-		// If this is being set from PIE we need to remove the pie prefix and point to the editor object
-		if (ObjectPackage->PIEInstanceID != INDEX_NONE)
+		UPackage* ObjectPackage = InObject->GetOutermost();
+
+		if (ensure(ObjectPackage))
 		{
-			FString PIEPrefix = FString::Printf(PLAYWORLD_PACKAGE_PREFIX TEXT("_%d_"), ObjectPackage->PIEInstanceID);
-			PackageName.ReplaceInline(*PIEPrefix, TEXT(""));
+			// If this is being set from PIE we need to remove the pie prefix and point to the editor object
+			if (ObjectPackage->PIEInstanceID != INDEX_NONE)
+			{
+				FString PIEPrefix = FString::Printf(PLAYWORLD_PACKAGE_PREFIX TEXT("_%d_"), ObjectPackage->PIEInstanceID);
+				FullPath.ReplaceInline(*PIEPrefix, TEXT(""));
+			}
 		}
 #endif
 		
-		FString FullPath = PackageName + TEXT(".") + InObject->GetPathName(ObjectPackage);
 		ExternalObjectPath = FSoftObjectPath(FullPath);
 	}
 }
