Sequencer bindings will always try to resolve before trying to load the referenced assets.

diff --git a/Engine/Source/Runtime/LevelSequence/Private/LevelSequenceBindingReference.cpp b/Engine/Source/Runtime/LevelSequence/Private/LevelSequenceBindingReference.cpp
--- a/Engine/Source/Runtime/LevelSequence/Private/LevelSequenceBindingReference.cpp
+++ b/Engine/Source/Runtime/LevelSequence/Private/LevelSequenceBindingReference.cpp
@@ -79,7 +79,14 @@ UObject* FLevelSequenceBindingReference::Resolve(UObject* InContext, FName Strea
 		}
 	#endif
 
-		return TempPath.TryLoad();
+		UObject* Object = TempPath.ResolveObject();
+
+		if (!Object)
+		{
+			Object = TempPath.TryLoad();
+		}
+
+		return Object;
 	}
 
 	return nullptr;
