Sequencer: only try to load external actor references from the editor (not in PIE).

diff --git a/Engine/Source/Runtime/LevelSequence/Private/LevelSequenceBindingReference.cpp b/Engine/Source/Runtime/LevelSequence/Private/LevelSequenceBindingReference.cpp
--- a/Engine/Source/Runtime/LevelSequence/Private/LevelSequenceBindingReference.cpp
+++ b/Engine/Source/Runtime/LevelSequence/Private/LevelSequenceBindingReference.cpp
@@ -81,11 +81,16 @@ UObject* FLevelSequenceBindingReference::Resolve(UObject* InContext, FName Strea
 
 		UObject* Object = TempPath.ResolveObject();
 
+#if WITH_EDITORONLY_DATA
 		if (!Object)
 		{
-			Object = TempPath.TryLoad();
+			UWorld* WorldContext = InContext->GetWorld();
+			if (!WorldContext || !WorldContext->IsPlayInEditor())
+			{
+				Object = TempPath.TryLoad();
+			}
 		}
-
+#endif
 		return Object;
 	}
 
