Fix UE4 Static Analysis Win64

diff --git a/Engine/Source/Runtime/LevelSequence/Private/LevelSequenceBindingReference.cpp b/Engine/Source/Runtime/LevelSequence/Private/LevelSequenceBindingReference.cpp
--- a/Engine/Source/Runtime/LevelSequence/Private/LevelSequenceBindingReference.cpp
+++ b/Engine/Source/Runtime/LevelSequence/Private/LevelSequenceBindingReference.cpp
@@ -84,7 +84,7 @@ UObject* FLevelSequenceBindingReference::Resolve(UObject* InContext, FName Strea
 #if WITH_EDITORONLY_DATA
 		if (!Object)
 		{
-			UWorld* WorldContext = InContext->GetWorld();
+			UWorld* WorldContext = InContext ? InContext->GetWorld() : nullptr;
 			if (!WorldContext || !WorldContext->IsPlayInEditor())
 			{
 				Object = TempPath.TryLoad();
