Sequencer: Hide TryLoad behavior for unresolved level sequence bindings behind a cvar (Sequencer.TryLoadUnresolvedReferences)

  - By default it is still enabled, but should be disabled for anyone that no longer relies on this temporary behavior, particularly when migrating sequence content to different levels

diff --git a/Engine/Source/Runtime/LevelSequence/Private/LevelSequenceBindingReference.cpp b/Engine/Source/Runtime/LevelSequence/Private/LevelSequenceBindingReference.cpp
--- a/Engine/Source/Runtime/LevelSequence/Private/LevelSequenceBindingReference.cpp
+++ b/Engine/Source/Runtime/LevelSequence/Private/LevelSequenceBindingReference.cpp
@@ -11,6 +11,26 @@
 #include "Components/SkeletalMeshComponent.h"
 #include "Engine/LevelStreamingDynamic.h"
 
+namespace UE
+{
+namespace MovieScene
+{
+
+#if WITH_EDITORONLY_DATA
+
+	int32 GTryLoadUnresolvedReferences = 1;
+	FAutoConsoleVariableRef GTryLoadUnresolvedReferencesCVar(
+		TEXT("Sequencer.TryLoadUnresolvedReferences"),
+		GTryLoadUnresolvedReferences,
+		TEXT("(Default: 1) When enabled, Sequencer will attempt to load the package for unresolved bindings in editor worlds. This feature will be removed in future since it can cause Actors from unrelated packages to be loaded into memory.\n"),
+		ECVF_Default
+	);
+
+#endif // WITH_EDITORONLY_DATA
+
+} // namespace MovieScene
+} // namespace UE
+
 FLevelSequenceBindingReference::FLevelSequenceBindingReference(UObject* InObject, UObject* InContext)
 {
 	check(InContext && InObject);
@@ -82,7 +102,7 @@ UObject* FLevelSequenceBindingReference::Resolve(UObject* InContext, FName Strea
 		UObject* Object = TempPath.ResolveObject();
 
 #if WITH_EDITORONLY_DATA
-		if (!Object)
+		if (Object == nullptr && UE::MovieScene::GTryLoadUnresolvedReferences != 0)
 		{
 			UWorld* WorldContext = InContext ? InContext->GetWorld() : nullptr;
 			if (!WorldContext || !WorldContext->IsPlayInEditor())
