Sequencer: Default TryLoadUnresolvedReferences to 0

diff --git a/Engine/Source/Runtime/LevelSequence/Private/LevelSequenceBindingReference.cpp b/Engine/Source/Runtime/LevelSequence/Private/LevelSequenceBindingReference.cpp
--- a/Engine/Source/Runtime/LevelSequence/Private/LevelSequenceBindingReference.cpp
+++ b/Engine/Source/Runtime/LevelSequence/Private/LevelSequenceBindingReference.cpp
@@ -18,11 +18,11 @@ namespace MovieScene
 
 #if WITH_EDITORONLY_DATA
 
-	int32 GTryLoadUnresolvedReferences = 1;
+	int32 GTryLoadUnresolvedReferences = 0;
 	FAutoConsoleVariableRef GTryLoadUnresolvedReferencesCVar(
 		TEXT("Sequencer.TryLoadUnresolvedReferences"),
 		GTryLoadUnresolvedReferences,
-		TEXT("(Default: 1) When enabled, Sequencer will attempt to load the package for unresolved bindings in editor worlds. This feature will be removed in future since it can cause Actors from unrelated packages to be loaded into memory.\n"),
+		TEXT("(Default: 0) When enabled, Sequencer will attempt to load the package for unresolved bindings in editor worlds. This feature will be removed in future since it can cause Actors from unrelated packages to be loaded into memory.\n"),
 		ECVF_Default
 	);
 
