Sequencer: Remove the TryLoad behavior for unresolved level sequence bindings.

diff --git a/Engine/Source/Runtime/LevelSequence/Private/LevelSequenceBindingReference.cpp b/Engine/Source/Runtime/LevelSequence/Private/LevelSequenceBindingReference.cpp
--- a/Engine/Source/Runtime/LevelSequence/Private/LevelSequenceBindingReference.cpp
+++ b/Engine/Source/Runtime/LevelSequence/Private/LevelSequenceBindingReference.cpp
@@ -11,26 +11,6 @@
 #include "Components/SkeletalMeshComponent.h"
 #include "Engine/LevelStreamingDynamic.h"
 
-namespace UE
-{
-namespace MovieScene
-{
-
-#if WITH_EDITORONLY_DATA
-
-	int32 GTryLoadUnresolvedReferences = 0;
-	FAutoConsoleVariableRef GTryLoadUnresolvedReferencesCVar(
-		TEXT("Sequencer.TryLoadUnresolvedReferences"),
-		GTryLoadUnresolvedReferences,
-		TEXT("(Default: 0) When enabled, Sequencer will attempt to load the package for unresolved bindings in editor worlds. This feature will be removed in future since it can cause Actors from unrelated packages to be loaded into memory.\n"),
-		ECVF_Default
-	);
-
-#endif // WITH_EDITORONLY_DATA
-
-} // namespace MovieScene
-} // namespace UE
-
 FLevelSequenceBindingReference::FLevelSequenceBindingReference(UObject* InObject, UObject* InContext)
 {
 	check(InContext && InObject);
@@ -99,19 +79,7 @@ UObject* FLevelSequenceBindingReference::Resolve(UObject* InContext, FName Strea
 		}
 	#endif
 
-		UObject* Object = TempPath.ResolveObject();
-
-#if WITH_EDITORONLY_DATA
-		if (Object == nullptr && UE::MovieScene::GTryLoadUnresolvedReferences != 0)
-		{
-			UWorld* WorldContext = InContext ? InContext->GetWorld() : nullptr;
-			if (!WorldContext || !WorldContext->IsPlayInEditor())
-			{
-				Object = TempPath.TryLoad();
-			}
-		}
-#endif
-		return Object;
+		return TempPath.ResolveObject();
 	}
 
 	return nullptr;
