Sequencer: Set the director blueprint name based on the level sequence (update it also on PostLoad). This makes it so that the blueprint tab will read as "NewLevelSequence (Director BP)" instead of repeatedly showing up as "SequenceDirector"

diff --git a/Engine/Plugins/MovieScene/LevelSequenceEditor/Source/LevelSequenceEditor/Private/Misc/MovieSceneSequenceEditor_LevelSequence.h b/Engine/Plugins/MovieScene/LevelSequenceEditor/Source/LevelSequenceEditor/Private/Misc/MovieSceneSequenceEditor_LevelSequence.h
--- a/Engine/Plugins/MovieScene/LevelSequenceEditor/Source/LevelSequenceEditor/Private/Misc/MovieSceneSequenceEditor_LevelSequence.h
+++ b/Engine/Plugins/MovieScene/LevelSequenceEditor/Source/LevelSequenceEditor/Private/Misc/MovieSceneSequenceEditor_LevelSequence.h
@@ -32,8 +32,7 @@ struct FMovieSceneSequenceEditor_LevelSequence : FMovieSceneSequenceEditor
 
 		ULevelSequence* LevelSequence = CastChecked<ULevelSequence>(InSequence);
 
-		FName BlueprintName = "SequenceDirector";
-		Blueprint = FKismetEditorUtilities::CreateBlueprint(ULevelSequenceDirector::StaticClass(), InSequence, BlueprintName, BPTYPE_Normal, UBlueprint::StaticClass(), UBlueprintGeneratedClass::StaticClass());
+		Blueprint = FKismetEditorUtilities::CreateBlueprint(ULevelSequenceDirector::StaticClass(), InSequence, FName(*LevelSequence->GetDirectorBlueprintName()), BPTYPE_Normal, UBlueprint::StaticClass(), UBlueprintGeneratedClass::StaticClass());
 		Blueprint->ClearFlags(RF_Standalone);
 
 		LevelSequence->SetDirectorBlueprint(Blueprint);
diff --git a/Engine/Source/Runtime/LevelSequence/Private/LevelSequence.cpp b/Engine/Source/Runtime/LevelSequence/Private/LevelSequence.cpp
--- a/Engine/Source/Runtime/LevelSequence/Private/LevelSequence.cpp
+++ b/Engine/Source/Runtime/LevelSequence/Private/LevelSequence.cpp
@@ -253,6 +253,8 @@ void ULevelSequence::PostLoad()
 		// Remove the binding for the director blueprint recompilation and re-add it to be sure there is only one entry in the list
 		DirectorBlueprint->OnCompiled().RemoveAll(this);
 		DirectorBlueprint->OnCompiled().AddUObject(this, &ULevelSequence::OnDirectorRecompiled);
+
+		DirectorBlueprint->Rename(*GetDirectorBlueprintName());
 	}
 
 	TSet<FGuid> InvalidSpawnables;
@@ -318,7 +320,7 @@ bool ULevelSequence::Rename(const TCHAR* NewName, UObject* NewOuter, ERenameFlag
 #if WITH_EDITOR
 	if (DirectorBlueprint)
 	{
-		DirectorBlueprint->Rename(*DirectorBlueprint->GetName(), this, Flags);
+		DirectorBlueprint->Rename(*GetDirectorBlueprintName(), this, Flags);
 	}
 #endif
 
@@ -454,6 +456,11 @@ UBlueprint* ULevelSequence::GetDirectorBlueprint() const
 	return DirectorBlueprint;
 }
 
+FString ULevelSequence::GetDirectorBlueprintName() const
+{
+	return GetDisplayName().ToString() + " (Director BP)";
+}
+
 void ULevelSequence::SetDirectorBlueprint(UBlueprint* NewDirectorBlueprint)
 {
 	if (DirectorBlueprint)
diff --git a/Engine/Source/Runtime/LevelSequence/Public/LevelSequence.h b/Engine/Source/Runtime/LevelSequence/Public/LevelSequence.h
--- a/Engine/Source/Runtime/LevelSequence/Public/LevelSequence.h
+++ b/Engine/Source/Runtime/LevelSequence/Public/LevelSequence.h
@@ -92,6 +92,8 @@ public:
 	 */
 	UBlueprint* GetDirectorBlueprint() const;
 
+	FString GetDirectorBlueprintName() const;
+
 protected:
 
 	virtual FGuid CreatePossessable(UObject* ObjectToPossess) override;
