#Sequencer Pass in rename flags to prevent resetting loaders during PostLoad

diff --git a/Engine/Source/Runtime/LevelSequence/Private/LevelSequence.cpp b/Engine/Source/Runtime/LevelSequence/Private/LevelSequence.cpp
--- a/Engine/Source/Runtime/LevelSequence/Private/LevelSequence.cpp
+++ b/Engine/Source/Runtime/LevelSequence/Private/LevelSequence.cpp
@@ -254,7 +254,7 @@ void ULevelSequence::PostLoad()
 		DirectorBlueprint->OnCompiled().RemoveAll(this);
 		DirectorBlueprint->OnCompiled().AddUObject(this, &ULevelSequence::OnDirectorRecompiled);
 
-		DirectorBlueprint->Rename(*GetDirectorBlueprintName());
+		DirectorBlueprint->Rename(*GetDirectorBlueprintName(), nullptr, (REN_NonTransactional|REN_ForceNoResetLoaders|REN_DoNotDirty));
 	}
 
 	TSet<FGuid> InvalidSpawnables;
