Sequencer: Move ShowBurnin/HideBurnin to Play and Stop so that the burnin isn't initialized for all time

diff --git a/Engine/Source/Runtime/LevelSequence/Private/LevelSequenceActor.cpp b/Engine/Source/Runtime/LevelSequence/Private/LevelSequenceActor.cpp
--- a/Engine/Source/Runtime/LevelSequence/Private/LevelSequenceActor.cpp
+++ b/Engine/Source/Runtime/LevelSequence/Private/LevelSequenceActor.cpp
@@ -65,7 +65,9 @@ ALevelSequenceActor::ALevelSequenceActor(const FObjectInitializer& Init)
 
 	// SequencePlayer must be a default sub object for it to be replicated correctly
 	SequencePlayer = Init.CreateDefaultSubobject<ULevelSequencePlayer>(this, "AnimationPlayer");
-
+	SequencePlayer->OnPlay.AddDynamic(this, &ALevelSequenceActor::ShowBurnin);
+	SequencePlayer->OnPlayReverse.AddDynamic(this, &ALevelSequenceActor::ShowBurnin);
+	SequencePlayer->OnStop.AddDynamic(this, &ALevelSequenceActor::HideBurnin);
 	bOverrideInstanceData = false;
 
 	// The level sequence actor defaults to never ticking by the tick manager because it is ticked separately in LevelTick
@@ -167,8 +169,6 @@ void ALevelSequenceActor::BeginPlay()
 
 	Super::BeginPlay();
 
-	RefreshBurnIn();
-
 	if (PlaybackSettings.bAutoPlay)
 	{
 		SequencePlayer->Play();
@@ -184,6 +184,10 @@ void ALevelSequenceActor::EndPlay(const EEndPlayReason::Type EndPlayReason)
 		// of waiting for the UObject to be destroyed by GC.
 		SequencePlayer->Stop();
 
+		SequencePlayer->OnPlay.RemoveAll(this);
+		SequencePlayer->OnPlayReverse.RemoveAll(this);
+		SequencePlayer->OnStop.RemoveAll(this);
+
 		if (UMovieSceneSequenceTickManager* TickManager = SequencePlayer->GetTickManager())
 		{
 			TickManager->UnregisterSequenceActor(this);
