Sequencer: Remove warnings that were to be removed with 17237666

diff --git a/Engine/Source/Runtime/LevelSequence/Private/LevelSequenceActor.cpp b/Engine/Source/Runtime/LevelSequence/Private/LevelSequenceActor.cpp
--- a/Engine/Source/Runtime/LevelSequence/Private/LevelSequenceActor.cpp
+++ b/Engine/Source/Runtime/LevelSequence/Private/LevelSequenceActor.cpp
@@ -311,22 +311,12 @@ void ALevelSequenceActor::HideBurnin()
 {
 	bShowBurnin = false;
 	RefreshBurnIn();
-
-	if (!BurnInOptions)
-	{
-		UE_LOG(LogLevelSequence, Warning, TEXT("Burnin is not enabled"));
-	}
 }
 
 void ALevelSequenceActor::ShowBurnin()
 {
 	bShowBurnin = true;
 	RefreshBurnIn();
-
-	if (!BurnInOptions || !BurnInOptions->bUseBurnIn)
-	{
-		UE_LOG(LogLevelSequence, Warning, TEXT("Burnin will not be visible because it is not enabled"));
-	}
 }
 
 void ALevelSequenceActor::RefreshBurnIn()
