Sequencer: Deprecated LevelSequenceSnapshotSetttings which was really only used by AutomatedLevelSequenceCapture to store the recording frame to be used in converting the snapshot times. The conversion is better done in the DefaultBurnin blueprint with the FrameRate using the FQualifiedFrameTime of the MasterTime or the ShotTime.

diff --git a/Engine/Source/Editor/MovieSceneTools/Private/AutomatedLevelSequenceCapture.cpp b/Engine/Source/Editor/MovieSceneTools/Private/AutomatedLevelSequenceCapture.cpp
--- a/Engine/Source/Editor/MovieSceneTools/Private/AutomatedLevelSequenceCapture.cpp
+++ b/Engine/Source/Editor/MovieSceneTools/Private/AutomatedLevelSequenceCapture.cpp
@@ -705,7 +705,6 @@ void UAutomatedLevelSequenceCapture::OnTick(float DeltaSeconds)
 	}
 	else if( CaptureState == ELevelSequenceCaptureState::ReadyToWarmUp )
 	{
-		Actor->SequencePlayer->SetSnapshotSettings(FLevelSequenceSnapshotSettings(Settings.ZeroPadFrameNumbers, Settings.GetFrameRate()));
 		Actor->SequencePlayer->Play();
 		// Start warming up
 		CaptureState = ELevelSequenceCaptureState::WarmingUp;
diff --git a/Engine/Source/Runtime/LevelSequence/Private/LevelSequencePlayer.cpp b/Engine/Source/Runtime/LevelSequence/Private/LevelSequencePlayer.cpp
--- a/Engine/Source/Runtime/LevelSequence/Private/LevelSequencePlayer.cpp
+++ b/Engine/Source/Runtime/LevelSequence/Private/LevelSequencePlayer.cpp
@@ -445,8 +445,6 @@ void ULevelSequencePlayer::TakeFrameSnapshot(FLevelSequencePlayerSnapshot& OutSn
 	// In Playback Resolution
 	const FFrameTime CurrentSequenceTime		  = ConvertFrameTime(CurrentPlayTime, PlayPosition.GetInputRate(), PlayPosition.GetOutputRate());
 
-	OutSnapshot.Settings = SnapshotSettings;
-
 	OutSnapshot.MasterTime = FQualifiedFrameTime(CurrentPlayTime, PlayPosition.GetInputRate());
 	OutSnapshot.MasterName = Sequence->GetName();
 
diff --git a/Engine/Source/Runtime/LevelSequence/Public/LevelSequencePlayer.h b/Engine/Source/Runtime/LevelSequence/Public/LevelSequencePlayer.h
--- a/Engine/Source/Runtime/LevelSequence/Public/LevelSequencePlayer.h
+++ b/Engine/Source/Runtime/LevelSequence/Public/LevelSequencePlayer.h
@@ -29,26 +29,8 @@ struct UE_DEPRECATED(4.15, "Please use FMovieSceneSequencePlaybackSettings.") FL
 DECLARE_DYNAMIC_MULTICAST_DELEGATE_OneParam(FOnLevelSequencePlayerCameraCutEvent, UCameraComponent*, CameraComponent);
 
 USTRUCT(BlueprintType)
-struct FLevelSequenceSnapshotSettings
-{
-	GENERATED_BODY()
-
-	FLevelSequenceSnapshotSettings()
-		: ZeroPadAmount(4), FrameRate(30, 1)
-	{}
-
-	FLevelSequenceSnapshotSettings(int32 InZeroPadAmount, FFrameRate InFrameRate)
-		: ZeroPadAmount(InZeroPadAmount), FrameRate(InFrameRate)
-	{}
-
-	/** Zero pad frames */
-	UPROPERTY(VisibleAnywhere, BlueprintReadOnly, Category="General")
-	uint8 ZeroPadAmount;
-
-	/** Playback framerate */
-	UPROPERTY(VisibleAnywhere, BlueprintReadOnly, Category="General")
-	FFrameRate FrameRate;
-};
+struct UE_DEPRECATED(5.0, "Snapshot settings are deprecated. Use the frame rate from the FQualifiedFrameTime on the MasterTime and the ShotTime") FLevelSequenceSnapshotSettings
+{ GENERATED_BODY() };
 
 /**
  * Frame snapshot information for a level sequence
@@ -82,9 +64,6 @@ struct FLevelSequencePlayerSnapshot
 	UPROPERTY(VisibleAnywhere, BlueprintReadOnly, Category="General")
 	TSoftObjectPtr<UCameraComponent> CameraComponent;
 
-	UPROPERTY(VisibleAnywhere, BlueprintReadOnly, Category = "General")
-	FLevelSequenceSnapshotSettings Settings;
-
 	UPROPERTY(VisibleAnywhere, BlueprintReadOnly, Category = "General")
 	TObjectPtr<ULevelSequence> ActiveShot = nullptr;
 
@@ -144,9 +123,6 @@ public:
 	UFUNCTION(BlueprintCallable, Category="Sequencer|Player", meta=(WorldContext="WorldContextObject", DynamicOutputParam="OutActor"))
 	static ULevelSequencePlayer* CreateLevelSequencePlayer(UObject* WorldContextObject, ULevelSequence* LevelSequence, FMovieSceneSequencePlaybackSettings Settings, ALevelSequenceActor*& OutActor);
 
-	/** Set the settings used to capture snapshots with */
-	void SetSnapshotSettings(const FLevelSequenceSnapshotSettings& InSettings) { SnapshotSettings = InSettings; }
-
 	/** Event triggered when there is a camera cut */
 	UPROPERTY(BlueprintAssignable, Category="Sequencer|Player")
 	FOnLevelSequencePlayerCameraCutEvent OnCameraCut;
@@ -219,9 +195,6 @@ private:
 
 protected:
 
-	/** How to take snapshots */
-	FLevelSequenceSnapshotSettings SnapshotSettings;
-
 	TOptional<int32> SnapshotOffsetTime;
 
 	TWeakObjectPtr<UCameraComponent> CachedCameraComponent;
