Sequencer: Control Rig; Need to TryLoad not ResolveObject with linked anim sequences since we need to load them to use them. UDN request, from another branch

diff --git a/Engine/Source/Runtime/LevelSequence/Private/AnimSequenceLevelSequenceLink.cpp b/Engine/Source/Runtime/LevelSequence/Private/AnimSequenceLevelSequenceLink.cpp
--- a/Engine/Source/Runtime/LevelSequence/Private/AnimSequenceLevelSequenceLink.cpp
+++ b/Engine/Source/Runtime/LevelSequence/Private/AnimSequenceLevelSequenceLink.cpp
@@ -11,7 +11,7 @@ void UAnimSequenceLevelSequenceLink::SetLevelSequence(ULevelSequence* InLevelSeq
 
 ULevelSequence* UAnimSequenceLevelSequenceLink::ResolveLevelSequence()
 {
-	UObject *Object = PathToLevelSequence.ResolveObject();
+	UObject *Object = PathToLevelSequence.TryLoad();
 	return Cast<ULevelSequence>(Object);
 }
 
diff --git a/Engine/Source/Runtime/LevelSequence/Private/LevelSequenceAnimSequenceLink.cpp b/Engine/Source/Runtime/LevelSequence/Private/LevelSequenceAnimSequenceLink.cpp
--- a/Engine/Source/Runtime/LevelSequence/Private/LevelSequenceAnimSequenceLink.cpp
+++ b/Engine/Source/Runtime/LevelSequence/Private/LevelSequenceAnimSequenceLink.cpp
@@ -11,7 +11,7 @@ void FLevelSequenceAnimSequenceLinkItem::SetAnimSequence(UAnimSequence* InAnimSe
 
 UAnimSequence* FLevelSequenceAnimSequenceLinkItem::ResolveAnimSequence()
 {
-	UObject *Object = PathToAnimSequence.ResolveObject();
+	UObject *Object = PathToAnimSequence.TryLoad();
 	return Cast<UAnimSequence>(Object);
 }
 
