Sequencer: Improved support for object binding IDs inside Level Sequence Directors
  - New methods for resolving bindings can now be used in-editor as well as in game and PIE: GetBoundObjects, GetBoundObject, GetBoundActors and GetBoundActor
  - Level Sequence Directors now know the sequence ID and native player interface they exist for - this allows them to correctly resolve relative binding IDs (something not possible through the top-level IMovieScenePlayer API)
  - The "Get Sequence Binding" BP node now defaults to the sequence that the BP exists within, if the BP exists inside a MovieSceneSequence. This will help resolve errors with people using the wrong binding.
  - The "Get Sequence Binding" node now always creates 'relative' bindings. Relative bindings at the root behave the same way as fixed bindings, so this won't change any behavior.

diff --git a/Engine/Plugins/MovieScene/ActorSequence/Source/ActorSequence/Private/ActorSequence.cpp b/Engine/Plugins/MovieScene/ActorSequence/Source/ActorSequence/Private/ActorSequence.cpp
--- a/Engine/Plugins/MovieScene/ActorSequence/Source/ActorSequence/Private/ActorSequence.cpp
+++ b/Engine/Plugins/MovieScene/ActorSequence/Source/ActorSequence/Private/ActorSequence.cpp
@@ -168,7 +168,7 @@ void UActorSequence::UnbindPossessableObjects(const FGuid& ObjectId)
 	ObjectReferences.RemoveBinding(ObjectId);
 }
 
-UObject* UActorSequence::CreateDirectorInstance(IMovieScenePlayer& Player)
+UObject* UActorSequence::CreateDirectorInstance(IMovieScenePlayer& Player, FMovieSceneSequenceID SequenceID)
 {
 	AActor* Actor = CastChecked<AActor>(Player.GetPlaybackContext(), ECastCheckedType::NullAllowed);
 	if (!Actor)
diff --git a/Engine/Plugins/MovieScene/ActorSequence/Source/ActorSequence/Public/ActorSequence.h b/Engine/Plugins/MovieScene/ActorSequence/Source/ActorSequence/Public/ActorSequence.h
--- a/Engine/Plugins/MovieScene/ActorSequence/Source/ActorSequence/Public/ActorSequence.h
+++ b/Engine/Plugins/MovieScene/ActorSequence/Source/ActorSequence/Public/ActorSequence.h
@@ -28,7 +28,7 @@ public:
 	virtual void UnbindPossessableObjects(const FGuid& ObjectId) override;
 	virtual void UnbindObjects(const FGuid& ObjectId, const TArray<UObject*>& InObjects, UObject* Context) override {}
 	virtual void UnbindInvalidObjects(const FGuid& ObjectId, UObject* Context) override {}
-	virtual UObject* CreateDirectorInstance(IMovieScenePlayer& Player) override;
+	virtual UObject* CreateDirectorInstance(IMovieScenePlayer& Player, FMovieSceneSequenceID SequenceID) override;
 
 #if WITH_EDITOR
 	virtual FText GetDisplayName() const override;
diff --git a/Engine/Source/Editor/MovieSceneTools/Private/K2Node_GetSequenceBinding.cpp b/Engine/Source/Editor/MovieSceneTools/Private/K2Node_GetSequenceBinding.cpp
--- a/Engine/Source/Editor/MovieSceneTools/Private/K2Node_GetSequenceBinding.cpp
+++ b/Engine/Source/Editor/MovieSceneTools/Private/K2Node_GetSequenceBinding.cpp
@@ -128,6 +128,16 @@ void UK2Node_GetSequenceBinding::AllocateDefaultPins()
 	Super::AllocateDefaultPins();
 }
 
+void UK2Node_GetSequenceBinding::PostPlacedNewNode()
+{
+	// Attempt to assign the sequence asset from our outer if this BP is contained within a sequence
+	if (UMovieSceneSequence* OuterSequence = GetTypedOuter<UMovieSceneSequence>())
+	{
+		SourceSequence = OuterSequence;
+	}
+	Super::PostPlacedNewNode();
+}
+
 UMovieScene* UK2Node_GetSequenceBinding::GetObjectMovieScene() const
 {
 	UMovieSceneSequence* Sequence = GetSequence();
@@ -310,6 +320,10 @@ TSharedPtr<SGraphNode> UK2Node_GetSequenceBinding::CreateVisualWidget()
 		SLATE_BEGIN_ARGS(SGraphNodeGetSequenceBinding){}
 		SLATE_END_ARGS()
 
+		SGraphNodeGetSequenceBinding()
+			: FMovieSceneObjectBindingIDPicker(MovieSceneSequenceID::Root, nullptr)
+		{}
+
 		void Construct(const FArguments& InArgs, UK2Node_GetSequenceBinding* InNode)
 		{
 			bNeedsUpdate = false;
diff --git a/Engine/Source/Editor/MovieSceneTools/Public/K2Node_GetSequenceBinding.h b/Engine/Source/Editor/MovieSceneTools/Public/K2Node_GetSequenceBinding.h
--- a/Engine/Source/Editor/MovieSceneTools/Public/K2Node_GetSequenceBinding.h
+++ b/Engine/Source/Editor/MovieSceneTools/Public/K2Node_GetSequenceBinding.h
@@ -48,6 +48,7 @@ public:
 	virtual void PreloadRequiredAssets() override;
 	virtual void GetMenuActions(FBlueprintActionDatabaseRegistrar& ActionRegistrar) const override;
 	virtual void AllocateDefaultPins() override;
+	virtual void PostPlacedNewNode() override;
 #if WITH_EDITOR
 	virtual TSharedPtr<SGraphNode> CreateVisualWidget() override;
 #endif
diff --git a/Engine/Source/Runtime/LevelSequence/Private/LevelSequence.cpp b/Engine/Source/Runtime/LevelSequence/Private/LevelSequence.cpp
--- a/Engine/Source/Runtime/LevelSequence/Private/LevelSequence.cpp
+++ b/Engine/Source/Runtime/LevelSequence/Private/LevelSequence.cpp
@@ -642,7 +642,7 @@ FGuid ULevelSequence::CreateSpawnable(UObject* ObjectToSpawn)
 
 #endif // WITH_EDITOR
 
-UObject* ULevelSequence::CreateDirectorInstance(IMovieScenePlayer& Player)
+UObject* ULevelSequence::CreateDirectorInstance(IMovieScenePlayer& Player, FMovieSceneSequenceID SequenceID)
 {
 	ULevelSequencePlayer* LevelSequencePlayer = Cast<ULevelSequencePlayer>(Player.AsUObject());
 	UObject*              DirectorOuter       = LevelSequencePlayer ? LevelSequencePlayer : Player.GetPlaybackContext();
@@ -658,6 +658,8 @@ UObject* ULevelSequence::CreateDirectorInstance(IMovieScenePlayer& Player)
 
 		ULevelSequenceDirector* NewDirector = NewObject<ULevelSequenceDirector>(DirectorOuter, DirectorClass, DirectorName, RF_Transient);
 		NewDirector->Player = LevelSequencePlayer;
+		NewDirector->MovieScenePlayerIndex = Player.GetUniqueIndex();
+		NewDirector->SubSequenceID = SequenceID.GetInternalValue();
 		NewDirector->OnCreated();
 		return NewDirector;
 	}
diff --git a/Engine/Source/Runtime/LevelSequence/Private/LevelSequenceDirector.cpp b/Engine/Source/Runtime/LevelSequence/Private/LevelSequenceDirector.cpp
--- a/Engine/Source/Runtime/LevelSequence/Private/LevelSequenceDirector.cpp
+++ b/Engine/Source/Runtime/LevelSequence/Private/LevelSequenceDirector.cpp
@@ -2,6 +2,9 @@
 
 #include "LevelSequenceDirector.h"
 #include "Engine/World.h"
+#include "UObject/Stack.h"
+#include "MovieSceneObjectBindingID.h"
+#include "IMovieScenePlayer.h"
 
 UWorld* ULevelSequenceDirector::GetWorld() const
 {
@@ -11,3 +14,88 @@ UWorld* ULevelSequenceDirector::GetWorld() const
 	}
 	return GetTypedOuter<UWorld>();
 }
+
+TArray<UObject*> ULevelSequenceDirector::GetBoundObjects(FMovieSceneObjectBindingID ObjectBinding)
+{
+	TArray<UObject*> Objects;
+
+	if (IMovieScenePlayer* PlayerInterface = IMovieScenePlayer::Get(static_cast<uint16>(MovieScenePlayerIndex)))
+	{
+		for (TWeakObjectPtr<> WeakObject : ObjectBinding.ResolveBoundObjects(FMovieSceneSequenceID(SubSequenceID), *PlayerInterface))
+		{
+			if (UObject* Object = WeakObject.Get())
+			{
+				Objects.Add(Object);
+			}
+		}
+	}
+	else
+	{
+		FFrame::KismetExecutionMessage(TEXT("No player interface available or assigned."), ELogVerbosity::Error);
+	}
+
+	return Objects;
+}
+
+UObject* ULevelSequenceDirector::GetBoundObject(FMovieSceneObjectBindingID ObjectBinding)
+{
+	if (IMovieScenePlayer* PlayerInterface = IMovieScenePlayer::Get(static_cast<uint16>(MovieScenePlayerIndex)))
+	{
+		for (TWeakObjectPtr<> WeakObject : ObjectBinding.ResolveBoundObjects(FMovieSceneSequenceID(SubSequenceID), *PlayerInterface))
+		{
+			if (UObject* Object = WeakObject.Get())
+			{
+				return Object;
+			}
+		}
+	}
+	else
+	{
+		FFrame::KismetExecutionMessage(TEXT("No player interface available or assigned."), ELogVerbosity::Error);
+	}
+
+	return nullptr;
+}
+
+TArray<AActor*> ULevelSequenceDirector::GetBoundActors(FMovieSceneObjectBindingID ObjectBinding)
+{
+	TArray<AActor*> Actors;
+
+	if (IMovieScenePlayer* PlayerInterface = IMovieScenePlayer::Get(static_cast<uint16>(MovieScenePlayerIndex)))
+	{
+		for (TWeakObjectPtr<> WeakObject : ObjectBinding.ResolveBoundObjects(FMovieSceneSequenceID(SubSequenceID), *PlayerInterface))
+		{
+			if (AActor* Actor = Cast<AActor>(WeakObject.Get()))
+			{
+				Actors.Add(Actor);
+			}
+		}
+	}
+	else
+	{
+		FFrame::KismetExecutionMessage(TEXT("No player interface available or assigned."), ELogVerbosity::Error);
+	}
+
+	return Actors;
+}
+
+AActor* ULevelSequenceDirector::GetBoundActor(FMovieSceneObjectBindingID ObjectBinding)
+{
+	if (IMovieScenePlayer* PlayerInterface = IMovieScenePlayer::Get(static_cast<uint16>(MovieScenePlayerIndex)))
+	{
+		for (TWeakObjectPtr<> WeakObject : ObjectBinding.ResolveBoundObjects(FMovieSceneSequenceID(SubSequenceID), *PlayerInterface))
+		{
+			if (AActor* Actor = Cast<AActor>(WeakObject.Get()))
+			{
+				return Actor;
+			}
+		}
+	}
+	else
+	{
+		FFrame::KismetExecutionMessage(TEXT("No player interface available or assigned."), ELogVerbosity::Error);
+	}
+
+	return nullptr;
+}
+
diff --git a/Engine/Source/Runtime/LevelSequence/Public/LevelSequence.h b/Engine/Source/Runtime/LevelSequence/Public/LevelSequence.h
--- a/Engine/Source/Runtime/LevelSequence/Public/LevelSequence.h
+++ b/Engine/Source/Runtime/LevelSequence/Public/LevelSequence.h
@@ -57,7 +57,7 @@ public:
 	virtual bool CanRebindPossessable(const FMovieScenePossessable& InPossessable) const override;
 	virtual UObject* MakeSpawnableTemplateFromInstance(UObject& InSourceObject, FName ObjectName) override;
 	virtual bool CanAnimateObject(UObject& InObject) const override;
-	virtual UObject* CreateDirectorInstance(IMovieScenePlayer& Player) override;
+	virtual UObject* CreateDirectorInstance(IMovieScenePlayer& Player, FMovieSceneSequenceID SequenceID) override;
 	virtual void PostLoad() override;
 	virtual bool Rename(const TCHAR* NewName = nullptr, UObject* NewOuter = nullptr, ERenameFlags Flags = REN_None) override;
 
diff --git a/Engine/Source/Runtime/LevelSequence/Public/LevelSequenceDirector.h b/Engine/Source/Runtime/LevelSequence/Public/LevelSequenceDirector.h
--- a/Engine/Source/Runtime/LevelSequence/Public/LevelSequenceDirector.h
+++ b/Engine/Source/Runtime/LevelSequence/Public/LevelSequenceDirector.h
@@ -5,8 +5,10 @@
 #include "CoreTypes.h"
 #include "Engine/Blueprint.h"
 #include "Engine/BlueprintGeneratedClass.h"
+#include "MovieSceneSequenceID.h"
 #include "LevelSequenceDirector.generated.h"
 
+class IMovieScenePlayer;
 class ULevelSequencePlayer;
 
 UCLASS(Blueprintable)
@@ -19,11 +21,61 @@ public:
 	UFUNCTION(BlueprintImplementableEvent, Category="Sequencer")
 	void OnCreated();
 
+
+	/**
+	 * Resolve the bindings inside this sub-sequence that relate to the specified ID
+	 * @note: ObjectBinding should be constructed from the same sequence as this Sequence Director's owning Sequence (see the GetSequenceBinding node)
+	 *
+	 * @param ObjectBinding The ID for the object binding inside this sub-sequence or one of its children to resolve
+	 */
+	UFUNCTION(BlueprintCallable, Category="Sequencer")
+	TArray<UObject*> GetBoundObjects(FMovieSceneObjectBindingID ObjectBinding);
+
+
+	/**
+	 * Resolve the first valid binding inside this sub-sequence that relates to the specified ID
+	 * @note: ObjectBinding should be constructed from the same sequence as this Sequence Director's owning Sequence (see the GetSequenceBinding node)
+	 *
+	 * @param ObjectBinding The ID for the object binding inside this sub-sequence or one of its children to resolve
+	 */
+	UFUNCTION(BlueprintCallable, Category="Sequencer")
+	UObject* GetBoundObject(FMovieSceneObjectBindingID ObjectBinding);
+
+
+	/**
+	 * Resolve the actor bindings inside this sub-sequence that relate to the specified ID
+	 * @note: ObjectBinding should be constructed from the same sequence as this Sequence Director's owning Sequence (see the GetSequenceBinding node)
+	 *
+	 * @param ObjectBinding The ID for the object binding inside this sub-sequence or one of its children to resolve
+	 */
+	UFUNCTION(BlueprintCallable, Category="Sequencer")
+	TArray<AActor*> GetBoundActors(FMovieSceneObjectBindingID ObjectBinding);
+
+
+	/**
+	 * Resolve the first valid Actor binding inside this sub-sequence that relates to the specified ID
+	 * @note: ObjectBinding should be constructed from the same sequence as this Sequence Director's owning Sequence (see the GetSequenceBinding node)
+	 *
+	 * @param ObjectBinding The ID for the object binding inside this sub-sequence or one of its children to resolve
+	 */
+	UFUNCTION(BlueprintCallable, Category="Sequencer")
+	AActor* GetBoundActor(FMovieSceneObjectBindingID ObjectBinding);
+
+public:
+
 	virtual UWorld* GetWorld() const override;
 
-	/** Pointer to the player that's playing back this director's sequence */
+	/** Pointer to the player that's playing back this director's sequence. Only valid in game or in PIE/Simulate. */
 	UPROPERTY(BlueprintReadOnly, Category="Cinematics")
 	TObjectPtr<ULevelSequencePlayer> Player;
+
+	/** The Sequence ID for the sequence this director is playing back within - has to be stored as an int32 so that it is reinstanced correctly*/
+	UPROPERTY()
+	int32 SubSequenceID;
+
+	/** Native player interface index - stored by index so that it can be reinstanced correctly */
+	UPROPERTY()
+	int32 MovieScenePlayerIndex;
 };
 
 
diff --git a/Engine/Source/Runtime/MovieScene/Private/Evaluation/MovieSceneEvaluationTemplateInstance.cpp b/Engine/Source/Runtime/MovieScene/Private/Evaluation/MovieSceneEvaluationTemplateInstance.cpp
--- a/Engine/Source/Runtime/MovieScene/Private/Evaluation/MovieSceneEvaluationTemplateInstance.cpp
+++ b/Engine/Source/Runtime/MovieScene/Private/Evaluation/MovieSceneEvaluationTemplateInstance.cpp
@@ -322,14 +322,14 @@ UObject* FMovieSceneRootEvaluationTemplateInstance::GetOrCreateDirectorInstance(
 	{
 		if (UMovieSceneSequence* Sequence = WeakRootSequence.Get())
 		{
-			NewDirectorInstance = Sequence->CreateDirectorInstance(Player);
+			NewDirectorInstance = Sequence->CreateDirectorInstance(Player, SequenceID);
 		}
 	}
 	else if (const FMovieSceneSequenceHierarchy* Hierarchy = CompiledDataManager->FindHierarchy(CompiledDataID))
 	{
 		const FMovieSceneSubSequenceData* SubData = Hierarchy->FindSubData(SequenceID);
 		check(SubData);
-		NewDirectorInstance = SubData->GetSequence()->CreateDirectorInstance(Player);
+		NewDirectorInstance = SubData->GetSequence()->CreateDirectorInstance(Player, SequenceID);
 	}
 
 	if (NewDirectorInstance)
diff --git a/Engine/Source/Runtime/MovieScene/Public/MovieSceneSequence.h b/Engine/Source/Runtime/MovieScene/Public/MovieSceneSequence.h
--- a/Engine/Source/Runtime/MovieScene/Public/MovieSceneSequence.h
+++ b/Engine/Source/Runtime/MovieScene/Public/MovieSceneSequence.h
@@ -7,6 +7,7 @@
 #include "Misc/Guid.h"
 #include "MovieSceneSignedObject.h"
 #include "MovieSceneTrack.h"
+#include "MovieSceneSequenceID.h"
 #include "MovieSceneSequence.generated.h"
 
 class ITargetPlatform;
@@ -199,7 +200,7 @@ public:
 	/**
 	 * Called to retrieve or construct a director instance to be used for the specified player
 	 */
-	virtual UObject* CreateDirectorInstance(IMovieScenePlayer& Player) { return nullptr; }
+	virtual UObject* CreateDirectorInstance(IMovieScenePlayer& Player, FMovieSceneSequenceID SequenceID) { return nullptr; }
 
 	MOVIESCENE_API virtual EMovieSceneServerClientMask OverrideNetworkMask(EMovieSceneServerClientMask InDefaultMask) const;
 
diff --git a/Engine/Source/Runtime/UMG/Private/Animation/WidgetAnimation.cpp b/Engine/Source/Runtime/UMG/Private/Animation/WidgetAnimation.cpp
--- a/Engine/Source/Runtime/UMG/Private/Animation/WidgetAnimation.cpp
+++ b/Engine/Source/Runtime/UMG/Private/Animation/WidgetAnimation.cpp
@@ -252,7 +252,7 @@ UMovieScene* UWidgetAnimation::GetMovieScene() const
 	return MovieScene;
 }
 
-UObject* UWidgetAnimation::CreateDirectorInstance(IMovieScenePlayer& Player)
+UObject* UWidgetAnimation::CreateDirectorInstance(IMovieScenePlayer& Player, FMovieSceneSequenceID SequenceID)
 {
 	// Widget animations do not create separate director instances, but just re-use the UUserWidget from the playback context
 	UUserWidget* WidgetContext = CastChecked<UUserWidget>(Player.GetPlaybackContext());
diff --git a/Engine/Source/Runtime/UMG/Public/Animation/WidgetAnimation.h b/Engine/Source/Runtime/UMG/Public/Animation/WidgetAnimation.h
--- a/Engine/Source/Runtime/UMG/Public/Animation/WidgetAnimation.h
+++ b/Engine/Source/Runtime/UMG/Public/Animation/WidgetAnimation.h
@@ -111,7 +111,7 @@ public:
 	virtual void UnbindObjects(const FGuid& ObjectId, const TArray<UObject*>& InObjects, UObject* InContext) override {}
 	virtual void UnbindInvalidObjects(const FGuid& ObjectId, UObject* InContext) override {}
 	virtual void LocateBoundObjects(const FGuid& ObjectId, UObject* Context, TArray<UObject*, TInlineAllocator<1>>& OutObjects) const override;
-	virtual UObject* CreateDirectorInstance(IMovieScenePlayer& Player) override;
+	virtual UObject* CreateDirectorInstance(IMovieScenePlayer& Player, FMovieSceneSequenceID SequenceID) override;
 #if WITH_EDITOR
 	virtual ETrackSupport IsTrackSupported(TSubclassOf<class UMovieSceneTrack> InTrackClass) const override;
 #endif
