Added include required for non-unity generated code

diff --git a/Engine/Source/Runtime/LevelSequence/Public/LevelSequenceDirector.h b/Engine/Source/Runtime/LevelSequence/Public/LevelSequenceDirector.h
--- a/Engine/Source/Runtime/LevelSequence/Public/LevelSequenceDirector.h
+++ b/Engine/Source/Runtime/LevelSequence/Public/LevelSequenceDirector.h
@@ -6,6 +6,7 @@
 #include "Engine/Blueprint.h"
 #include "Engine/BlueprintGeneratedClass.h"
 #include "MovieSceneSequenceID.h"
+#include "MovieSceneObjectBindingID.h"
 #include "LevelSequenceDirector.generated.h"
 
 class IMovieScenePlayer;
