Sequencer: Remove invalid binding references on PostLoad()

diff --git a/Engine/Source/Runtime/LevelSequence/Private/LevelSequence.cpp b/Engine/Source/Runtime/LevelSequence/Private/LevelSequence.cpp
--- a/Engine/Source/Runtime/LevelSequence/Private/LevelSequence.cpp
+++ b/Engine/Source/Runtime/LevelSequence/Private/LevelSequence.cpp
@@ -298,21 +298,20 @@ void ULevelSequence::PostLoad()
 				PurgeLegacyBlueprints(ObjectInPackage, Package);
 			}
 		}
+	}
 
-		// Remove any invalid object bindings
-		TSet<FGuid> ValidObjectBindings;
-		for (int32 Index = 0; Index < MovieScene->GetSpawnableCount(); ++Index)
-		{
-			ValidObjectBindings.Add(MovieScene->GetSpawnable(Index).GetGuid());
-		}
-		for (int32 Index = 0; Index < MovieScene->GetPossessableCount(); ++Index)
-		{
-			ValidObjectBindings.Add(MovieScene->GetPossessable(Index).GetGuid());
-		}
-
-		BindingReferences.RemoveInvalidBindings(ValidObjectBindings);
+	// Remove any invalid object bindings
+	TSet<FGuid> ValidObjectBindings;
+	for (int32 Index = 0; Index < MovieScene->GetSpawnableCount(); ++Index)
+	{
+		ValidObjectBindings.Add(MovieScene->GetSpawnable(Index).GetGuid());
+	}
+	for (int32 Index = 0; Index < MovieScene->GetPossessableCount(); ++Index)
+	{
+		ValidObjectBindings.Add(MovieScene->GetPossessable(Index).GetGuid());
 	}
 
+	BindingReferences.RemoveInvalidBindings(ValidObjectBindings);
 #endif
 }
 
