Sequencer: Attach and Path track should be supported by LevelSequence and not other sequences like TemplateSequence which don't support multiple bindings and cross level bindings. This fixes a crash when trying to make an attachment in a template sequence.

diff --git a/Engine/Source/Editor/MovieSceneTools/Private/TrackEditors/AttachTrackEditor.cpp b/Engine/Source/Editor/MovieSceneTools/Private/TrackEditors/AttachTrackEditor.cpp
--- a/Engine/Source/Editor/MovieSceneTools/Private/TrackEditors/AttachTrackEditor.cpp
+++ b/Engine/Source/Editor/MovieSceneTools/Private/TrackEditors/AttachTrackEditor.cpp
@@ -138,6 +138,11 @@ bool F3DAttachTrackEditor::SupportsType( TSubclassOf<UMovieSceneTrack> Type ) co
 	return Type == UMovieScene3DAttachTrack::StaticClass();
 }
 
+bool F3DAttachTrackEditor::SupportsSequence(UMovieSceneSequence* InSequence) const
+{
+	ETrackSupport TrackSupported = InSequence ? InSequence->IsTrackSupported(UMovieScene3DAttachTrack::StaticClass()) : ETrackSupport::NotSupported;
+	return TrackSupported == ETrackSupport::Supported;
+}
 
 TSharedRef<ISequencerSection> F3DAttachTrackEditor::MakeSectionInterface( UMovieSceneSection& SectionObject, UMovieSceneTrack& Track, FGuid ObjectBinding )
 {
diff --git a/Engine/Source/Editor/MovieSceneTools/Private/TrackEditors/AttachTrackEditor.h b/Engine/Source/Editor/MovieSceneTools/Private/TrackEditors/AttachTrackEditor.h
--- a/Engine/Source/Editor/MovieSceneTools/Private/TrackEditors/AttachTrackEditor.h
+++ b/Engine/Source/Editor/MovieSceneTools/Private/TrackEditors/AttachTrackEditor.h
@@ -67,6 +67,7 @@ public:
 	virtual void BuildObjectBindingTrackMenu(FMenuBuilder& MenuBuilder, const TArray<FGuid>& ObjectBindings, const UClass* ObjectClass) override;
 	virtual TSharedRef<ISequencerSection> MakeSectionInterface( UMovieSceneSection& SectionObject, UMovieSceneTrack& Track, FGuid ObjectBinding ) override;
 	virtual bool SupportsType( TSubclassOf<UMovieSceneTrack> Type ) const override;
+	virtual bool SupportsSequence(UMovieSceneSequence* InSequence) const override;
 
 	// FTrackEditorActorPicker
 	virtual bool IsActorPickable( const AActor* const ParentActor, FGuid ObjectBinding, UMovieSceneSection* InSection ) override;
diff --git a/Engine/Source/Editor/MovieSceneTools/Private/TrackEditors/PathTrackEditor.cpp b/Engine/Source/Editor/MovieSceneTools/Private/TrackEditors/PathTrackEditor.cpp
--- a/Engine/Source/Editor/MovieSceneTools/Private/TrackEditors/PathTrackEditor.cpp
+++ b/Engine/Source/Editor/MovieSceneTools/Private/TrackEditors/PathTrackEditor.cpp
@@ -100,6 +100,11 @@ bool F3DPathTrackEditor::SupportsType( TSubclassOf<UMovieSceneTrack> Type ) cons
 	return Type == UMovieScene3DPathTrack::StaticClass();
 }
 
+bool F3DPathTrackEditor::SupportsSequence(UMovieSceneSequence* InSequence) const
+{
+	ETrackSupport TrackSupported = InSequence ? InSequence->IsTrackSupported(UMovieScene3DPathTrack::StaticClass()) : ETrackSupport::NotSupported;
+	return TrackSupported == ETrackSupport::Supported;
+}
 
 TSharedRef<ISequencerSection> F3DPathTrackEditor::MakeSectionInterface( UMovieSceneSection& SectionObject, UMovieSceneTrack& Track, FGuid ObjectBinding )
 {
diff --git a/Engine/Source/Editor/MovieSceneTools/Private/TrackEditors/PathTrackEditor.h b/Engine/Source/Editor/MovieSceneTools/Private/TrackEditors/PathTrackEditor.h
--- a/Engine/Source/Editor/MovieSceneTools/Private/TrackEditors/PathTrackEditor.h
+++ b/Engine/Source/Editor/MovieSceneTools/Private/TrackEditors/PathTrackEditor.h
@@ -48,6 +48,7 @@ public:
 	virtual void BuildObjectBindingTrackMenu(FMenuBuilder& MenuBuilder, const TArray<FGuid>& ObjectBindings, const UClass* ObjectClass) override;
 	virtual TSharedRef<ISequencerSection> MakeSectionInterface( UMovieSceneSection& SectionObject, UMovieSceneTrack& Track, FGuid ObjectBinding ) override;
 	virtual bool SupportsType( TSubclassOf<UMovieSceneTrack> Type ) const override;
+	virtual bool SupportsSequence(UMovieSceneSequence* InSequence) const override;
 
 	// FTrackEditorActorPicker
 	virtual bool IsActorPickable( const AActor* const ParentActor, FGuid ObjectBinding, UMovieSceneSection* InSection ) override;
diff --git a/Engine/Source/Runtime/LevelSequence/Private/LevelSequence.cpp b/Engine/Source/Runtime/LevelSequence/Private/LevelSequence.cpp
--- a/Engine/Source/Runtime/LevelSequence/Private/LevelSequence.cpp
+++ b/Engine/Source/Runtime/LevelSequence/Private/LevelSequence.cpp
@@ -15,6 +15,8 @@
 #include "Animation/AnimInstance.h"
 #include "LevelSequenceModule.h"
 #include "MovieSceneSpawnableAnnotation.h"
+#include "Tracks/MovieScene3DAttachTrack.h"
+#include "Tracks/MovieScene3DPathTrack.h"
 #include "Tracks/MovieSceneAudioTrack.h"
 #include "Tracks/MovieSceneCameraCutTrack.h"
 #include "Tracks/MovieSceneCinematicShotTrack.h"
@@ -103,7 +105,9 @@ bool ULevelSequence::CanAnimateObject(UObject& InObject) const
 
 ETrackSupport ULevelSequence::IsTrackSupported(TSubclassOf<class UMovieSceneTrack> InTrackClass) const
 {
-	if (InTrackClass == UMovieSceneAudioTrack::StaticClass() ||
+	if (InTrackClass == UMovieScene3DAttachTrack::StaticClass() ||
+		InTrackClass == UMovieScene3DPathTrack::StaticClass() ||
+		InTrackClass == UMovieSceneAudioTrack::StaticClass() ||
 		InTrackClass == UMovieSceneCameraCutTrack::StaticClass() ||
 		InTrackClass == UMovieSceneCinematicShotTrack::StaticClass() ||
 		InTrackClass == UMovieSceneEventTrack::StaticClass() ||
