Sequencer: Redo RemoveInvalidBindings by calling it in PostInitProperties instead of PostLoad since the movie scene might not have been initialized yet. This was causing some automation tests to fail.

This came from another project where adding a binding, deleting it, saving and restarting would not remove the binding from the movie scene.

diff --git a/Engine/Source/Runtime/LevelSequence/Private/LevelSequence.cpp b/Engine/Source/Runtime/LevelSequence/Private/LevelSequence.cpp
--- a/Engine/Source/Runtime/LevelSequence/Private/LevelSequence.cpp
+++ b/Engine/Source/Runtime/LevelSequence/Private/LevelSequence.cpp
@@ -303,19 +303,29 @@ void ULevelSequence::PostLoad()
 			}
 		}
 	}
+#endif
+}
 
-	// Remove any invalid object bindings
-	TSet<FGuid> ValidObjectBindings;
-	for (int32 Index = 0; Index < MovieScene->GetSpawnableCount(); ++Index)
-	{
-		ValidObjectBindings.Add(MovieScene->GetSpawnable(Index).GetGuid());
-	}
-	for (int32 Index = 0; Index < MovieScene->GetPossessableCount(); ++Index)
+void ULevelSequence::PostInitProperties()
+{
+	Super::PostInitProperties();
+
+#if WITH_EDITOR
+	if (MovieScene)
 	{
-		ValidObjectBindings.Add(MovieScene->GetPossessable(Index).GetGuid());
-	}
+		// Remove any invalid object bindings
+		TSet<FGuid> ValidObjectBindings;
+		for (int32 Index = 0; Index < MovieScene->GetSpawnableCount(); ++Index)
+		{
+			ValidObjectBindings.Add(MovieScene->GetSpawnable(Index).GetGuid());
+		}
+		for (int32 Index = 0; Index < MovieScene->GetPossessableCount(); ++Index)
+		{
+			ValidObjectBindings.Add(MovieScene->GetPossessable(Index).GetGuid());
+		}
 
-	BindingReferences.RemoveInvalidBindings(ValidObjectBindings);
+		BindingReferences.RemoveInvalidBindings(ValidObjectBindings);
+	}
 #endif
 }
 
diff --git a/Engine/Source/Runtime/LevelSequence/Public/LevelSequence.h b/Engine/Source/Runtime/LevelSequence/Public/LevelSequence.h
--- a/Engine/Source/Runtime/LevelSequence/Public/LevelSequence.h
+++ b/Engine/Source/Runtime/LevelSequence/Public/LevelSequence.h
@@ -59,6 +59,7 @@ public:
 	virtual bool CanAnimateObject(UObject& InObject) const override;
 	virtual UObject* CreateDirectorInstance(IMovieScenePlayer& Player, FMovieSceneSequenceID SequenceID) override;
 	virtual void PostLoad() override;
+	virtual void PostInitProperties() override;
 	virtual bool Rename(const TCHAR* NewName = nullptr, UObject* NewOuter = nullptr, ERenameFlags Flags = REN_None) override;
 
 
