Sequencer: Disable auto activate for all UFXSystemComponents rather than just UParticleSystemComponent so that Niagara components also have auto activate disabled.

diff --git a/Engine/Source/Runtime/LevelSequence/Private/LevelSequenceActorSpawner.cpp b/Engine/Source/Runtime/LevelSequence/Private/LevelSequenceActorSpawner.cpp
--- a/Engine/Source/Runtime/LevelSequence/Private/LevelSequenceActorSpawner.cpp
+++ b/Engine/Source/Runtime/LevelSequence/Private/LevelSequenceActorSpawner.cpp
@@ -167,13 +167,13 @@ UObject* FLevelSequenceActorSpawner::SpawnObject(FMovieSceneSpawnable& Spawnable
 	}
 
 	{
-		// Disable all particle components so that they don't auto fire as soon as the actor is spawned. The particles should be triggered through the particle track.
+		// Disable all FXSystem components so that they don't auto fire as soon as the actor is spawned. The particles should be triggered through the particle track.
 		for (UActorComponent* Component : ObjectTemplate->GetComponents())
 		{
-			if (UParticleSystemComponent* ParticleComponent = Cast<UParticleSystemComponent>(Component))
+			if (UFXSystemComponent* FXSystemComponent = Cast<UFXSystemComponent>(Component))
 			{
 				// The particle needs to be set inactive in case its template was active.
-				ParticleComponent->SetActiveFlag(false);
+				FXSystemComponent->SetActiveFlag(false);
 				Component->bAutoActivate = false;
 			}
 		}
