[Backout] - CL18711725 which unified some aspects of cascade and niagara in sequencer, but broke existing content that was relying on tracks to be auto-activated

diff --git a/Engine/Source/Runtime/LevelSequence/Private/LevelSequenceActorSpawner.cpp b/Engine/Source/Runtime/LevelSequence/Private/LevelSequenceActorSpawner.cpp
--- a/Engine/Source/Runtime/LevelSequence/Private/LevelSequenceActorSpawner.cpp
+++ b/Engine/Source/Runtime/LevelSequence/Private/LevelSequenceActorSpawner.cpp
@@ -167,13 +167,13 @@ UObject* FLevelSequenceActorSpawner::SpawnObject(FMovieSceneSpawnable& Spawnable
 	}
 
 	{
-		// Disable all FXSystem components so that they don't auto fire as soon as the actor is spawned. The particles should be triggered through the particle track.
+		// Disable all particle components so that they don't auto fire as soon as the actor is spawned. The particles should be triggered through the particle track.
 		for (UActorComponent* Component : ObjectTemplate->GetComponents())
 		{
-			if (UFXSystemComponent* FXSystemComponent = Cast<UFXSystemComponent>(Component))
+			if (UParticleSystemComponent* ParticleComponent = Cast<UParticleSystemComponent>(Component))
 			{
 				// The particle needs to be set inactive in case its template was active.
-				FXSystemComponent->SetActiveFlag(false);
+				ParticleComponent->SetActiveFlag(false);
 				Component->bAutoActivate = false;
 			}
 		}
