Explicitely enforce bReplicateUsingRegisteredSubObjectList to false for classes that implemented ReplicateSubObjects but won't be converted to using the new list.

diff --git a/Engine/Source/Runtime/LevelSequence/Private/LevelSequenceActor.cpp b/Engine/Source/Runtime/LevelSequence/Private/LevelSequenceActor.cpp
--- a/Engine/Source/Runtime/LevelSequence/Private/LevelSequenceActor.cpp
+++ b/Engine/Source/Runtime/LevelSequence/Private/LevelSequenceActor.cpp
@@ -61,6 +61,8 @@ ALevelSequenceActor::ALevelSequenceActor(const FObjectInitializer& Init)
 	bIsSpatiallyLoaded = false;
 #endif //WITH_EDITORONLY_DATA
 
+	bReplicateUsingRegisteredSubObjectList = false;
+
 	BindingOverrides = Init.CreateDefaultSubobject<UMovieSceneBindingOverrides>(this, "BindingOverrides");
 	BurnInOptions = Init.CreateDefaultSubobject<ULevelSequenceBurnInOptions>(this, "BurnInOptions");
 	DefaultInstanceData = Init.CreateDefaultSubobject<UDefaultLevelSequenceInstanceData>(this, "InstanceData");
