Add logging of camera cuts in sequences.

This is to help with troubleshooting cut issues.

diff --git a/Engine/Source/Runtime/LevelSequence/Private/LevelSequencePlayer.cpp b/Engine/Source/Runtime/LevelSequence/Private/LevelSequencePlayer.cpp
--- a/Engine/Source/Runtime/LevelSequence/Private/LevelSequencePlayer.cpp
+++ b/Engine/Source/Runtime/LevelSequence/Private/LevelSequencePlayer.cpp
@@ -25,6 +25,7 @@
 #include "Modules/ModuleManager.h"
 #include "LevelUtils.h"
 #include "Core/Public/ProfilingDebugging/CsvProfiler.h"
+#include "LevelSequenceModule.h"
 
 /* ULevelSequencePlayer structors
  *****************************************************************************/
@@ -326,6 +327,11 @@ void ULevelSequencePlayer::UpdateCameraCut(UObject* CameraObject, const EMovieSc
 	FViewTargetTransitionParams TransitionParams;
 	if (CameraCutParams.BlendType.IsSet())
 	{
+		UE_LOG(LogLevelSequence, Log, TEXT("Blending into new camera cut: '%s' -> '%s' (blend time: %f)"),
+			(ViewTarget ? *ViewTarget->GetName() : TEXT("None")),
+			(CameraObject ? *CameraObject->GetName() : TEXT("None")),
+			TransitionParams.BlendTime);
+
 		// Convert known easing functions to their corresponding view target blend parameters.
 		TTuple<EViewTargetBlendFunction, float> BlendFunctionAndExp = BuiltInEasingTypeToBlendFunction(CameraCutParams.BlendType.GetValue());
 		TransitionParams.BlendTime = CameraCutParams.BlendTime;
@@ -342,10 +348,16 @@ void ULevelSequencePlayer::UpdateCameraCut(UObject* CameraObject, const EMovieSc
 			const AActor* PendingViewTarget = PC->PlayerCameraManager->PendingViewTarget.Target;
 			if (CameraActor != nullptr && PendingViewTarget == CameraActor)
 			{
+				UE_LOG(LogLevelSequence, Log, TEXT("Camera transition aborted, we are already blending towards the intended camera"));
 				bDoSetViewTarget = false;
 			}
 		}
 	}
+	else
+	{
+		UE_LOG(LogLevelSequence, Log, TEXT("Starting new camera cut: '%s'"),
+			(CameraObject ? *CameraObject->GetName() : TEXT("None")));
+	}
 	if (bDoSetViewTarget)
 	{
 		PC->SetViewTarget(CameraActor, TransitionParams);
