Fix issue where aspect ratio constraint wasn't correctly restored in some cases

diff --git a/Engine/Source/Runtime/LevelSequence/Private/LevelSequencePlayer.cpp b/Engine/Source/Runtime/LevelSequence/Private/LevelSequencePlayer.cpp
--- a/Engine/Source/Runtime/LevelSequence/Private/LevelSequencePlayer.cpp
+++ b/Engine/Source/Runtime/LevelSequence/Private/LevelSequencePlayer.cpp
@@ -293,6 +293,7 @@ void ULevelSequencePlayer::UpdateCameraCut(UObject* CameraObject, const EMovieSc
 
 	// override the player controller's view target
 	AActor* CameraActor = Cast<AActor>(CameraObject);
+	ULocalPlayer* LocalPlayer = PC->GetLocalPlayer();
 
 	// if the camera object is null, use the last view target so that it is restored to the state before the sequence takes control
 	bool bRestoreAspectRatioConstraint = false;
@@ -304,13 +305,15 @@ void ULevelSequencePlayer::UpdateCameraCut(UObject* CameraObject, const EMovieSc
 		// Skip if the last view target is the same as the current view target so that there's no additional camera cut
 		if (CameraActor == ViewTarget)
 		{
+			if (LocalPlayer && LastAspectRatioAxisConstraint.IsSet())
+			{
+				LocalPlayer->AspectRatioAxisConstraint = LastAspectRatioAxisConstraint.GetValue();
+			}
 			return;
 		}
 	}
 
 	// Save the last view target/aspect ratio constraint/etc. so that it can all be restored when the camera object is null.
-	ULocalPlayer* LocalPlayer = PC->GetLocalPlayer();
-
 	if (!LastViewTarget.IsValid())
 	{
 		LastViewTarget = ViewTarget;
