Sequencer: Added functions for getting the master sequence time, and local sub sequence time to Level Sequence Directors

diff --git a/Engine/Source/Runtime/LevelSequence/Private/LevelSequenceDirector.cpp b/Engine/Source/Runtime/LevelSequence/Private/LevelSequenceDirector.cpp
--- a/Engine/Source/Runtime/LevelSequence/Private/LevelSequenceDirector.cpp
+++ b/Engine/Source/Runtime/LevelSequence/Private/LevelSequenceDirector.cpp
@@ -7,6 +7,9 @@
 #include "LevelSequencePlayer.h"
 #include "IMovieScenePlayer.h"
 #include "Evaluation/MovieSceneEvaluationTemplateInstance.h"
+#include "MovieSceneSequence.h"
+#include "Misc/QualifiedFrameTime.h"
+#include "Evaluation/MovieScenePlayback.h"
 
 UWorld* ULevelSequenceDirector::GetWorld() const
 {
@@ -17,6 +20,52 @@ UWorld* ULevelSequenceDirector::GetWorld() const
 	return GetTypedOuter<UWorld>();
 }
 
+FQualifiedFrameTime ULevelSequenceDirector::GetMasterSequenceTime() const
+{
+	using namespace UE::MovieScene;
+
+	if (IMovieScenePlayer* PlayerInterface = IMovieScenePlayer::Get(static_cast<uint16>(MovieScenePlayerIndex)))
+	{
+		FMovieSceneRootEvaluationTemplateInstance& EvaluationTemplate = PlayerInterface->GetEvaluationTemplate();
+		UMovieSceneSequence* RootSequence = EvaluationTemplate.GetRootSequence();
+		const FSequenceInstance* RootSequenceInstance = EvaluationTemplate.FindInstance(MovieSceneSequenceID::Root);
+		if (RootSequenceInstance && RootSequence)
+		{
+			// Put the qualified frame time into 'display' rate
+			FFrameRate DisplayRate = RootSequence->GetMovieScene()->GetDisplayRate();
+			FMovieSceneContext Context = RootSequenceInstance->GetContext();
+
+			FFrameTime DisplayRateTime = ConvertFrameTime(Context.GetTime(), Context.GetFrameRate(), DisplayRate);
+			return FQualifiedFrameTime(DisplayRateTime, DisplayRate);
+		}
+	}
+	return FQualifiedFrameTime(0, FFrameRate());
+}
+
+FQualifiedFrameTime ULevelSequenceDirector::GetCurrentTime() const
+{
+	using namespace UE::MovieScene;
+
+	if (IMovieScenePlayer* PlayerInterface = IMovieScenePlayer::Get(static_cast<uint16>(MovieScenePlayerIndex)))
+	{
+		FMovieSceneRootEvaluationTemplateInstance& EvaluationTemplate = PlayerInterface->GetEvaluationTemplate();
+
+		UMovieSceneSequence* SubSequence = EvaluationTemplate.GetSequence(FMovieSceneSequenceID(SubSequenceID));
+		const FSequenceInstance* SequenceInstance = EvaluationTemplate.FindInstance(FMovieSceneSequenceID(SubSequenceID));
+		if (SequenceInstance && SubSequence)
+		{
+			// Put the qualified frame time into 'display' rate
+			FFrameRate DisplayRate = SubSequence->GetMovieScene()->GetDisplayRate();
+			FMovieSceneContext Context = SequenceInstance->GetContext();
+
+			FFrameTime DisplayRateTime = ConvertFrameTime(Context.GetTime(), Context.GetFrameRate(), DisplayRate);
+			return FQualifiedFrameTime(DisplayRateTime, DisplayRate);
+		}
+	}
+	return FQualifiedFrameTime(0, FFrameRate());
+}
+
+
 TArray<UObject*> ULevelSequenceDirector::GetBoundObjects(FMovieSceneObjectBindingID ObjectBinding)
 {
 	TArray<UObject*> Objects;
diff --git a/Engine/Source/Runtime/LevelSequence/Public/LevelSequenceDirector.h b/Engine/Source/Runtime/LevelSequence/Public/LevelSequenceDirector.h
--- a/Engine/Source/Runtime/LevelSequence/Public/LevelSequenceDirector.h
+++ b/Engine/Source/Runtime/LevelSequence/Public/LevelSequenceDirector.h
@@ -21,7 +21,20 @@ public:
 	/** Called when this director is created */
 	UFUNCTION(BlueprintImplementableEvent, Category="Sequencer")
 	void OnCreated();
+	
+	/**
+	 * Get the current time for the outermost (root, or master) sequence
+	 * @return The current playback position of the master sequence
+	 */
+	UFUNCTION(BlueprintCallable, Category = "Sequencer|Director")
+	FQualifiedFrameTime GetMasterSequenceTime() const;
 
+	/**
+	 * Get the current time for this director's sub-sequence (or the master sequence, if this is a master sequence director)
+	 * @return The current playback position of this director's sequence
+	 */
+	UFUNCTION(BlueprintCallable, Category = "Sequencer|Director")
+	FQualifiedFrameTime GetCurrentTime() const;
 
 	/**
 	 * Resolve the bindings inside this sub-sequence that relate to the specified ID
@@ -29,7 +42,7 @@ public:
 	 *
 	 * @param ObjectBinding The ID for the object binding inside this sub-sequence or one of its children to resolve
 	 */
-	UFUNCTION(BlueprintCallable, Category="Sequencer")
+	UFUNCTION(BlueprintCallable, Category="Sequencer|Director")
 	TArray<UObject*> GetBoundObjects(FMovieSceneObjectBindingID ObjectBinding);
 
 
@@ -39,7 +52,7 @@ public:
 	 *
 	 * @param ObjectBinding The ID for the object binding inside this sub-sequence or one of its children to resolve
 	 */
-	UFUNCTION(BlueprintCallable, Category="Sequencer")
+	UFUNCTION(BlueprintCallable, Category="Sequencer|Director")
 	UObject* GetBoundObject(FMovieSceneObjectBindingID ObjectBinding);
 
 
@@ -49,7 +62,7 @@ public:
 	 *
 	 * @param ObjectBinding The ID for the object binding inside this sub-sequence or one of its children to resolve
 	 */
-	UFUNCTION(BlueprintCallable, Category="Sequencer")
+	UFUNCTION(BlueprintCallable, Category="Sequencer|Director")
 	TArray<AActor*> GetBoundActors(FMovieSceneObjectBindingID ObjectBinding);
 
 
@@ -59,13 +72,13 @@ public:
 	 *
 	 * @param ObjectBinding The ID for the object binding inside this sub-sequence or one of its children to resolve
 	 */
-	UFUNCTION(BlueprintCallable, Category="Sequencer")
+	UFUNCTION(BlueprintCallable, Category="Sequencer|Director")
 	AActor* GetBoundActor(FMovieSceneObjectBindingID ObjectBinding);
 
 	/*
 	 * Get the current sequence that this director is playing back within 
 	 */
-	UFUNCTION(BlueprintCallable, Category="Sequencer")
+	UFUNCTION(BlueprintCallable, Category="Sequencer|Director")
 	UMovieSceneSequence* GetSequence();
 
 public:
