Moved #include to header to appease generated cpp

diff --git a/Engine/Source/Runtime/LevelSequence/Private/LevelSequenceDirector.cpp b/Engine/Source/Runtime/LevelSequence/Private/LevelSequenceDirector.cpp
--- a/Engine/Source/Runtime/LevelSequence/Private/LevelSequenceDirector.cpp
+++ b/Engine/Source/Runtime/LevelSequence/Private/LevelSequenceDirector.cpp
@@ -8,7 +8,6 @@
 #include "IMovieScenePlayer.h"
 #include "Evaluation/MovieSceneEvaluationTemplateInstance.h"
 #include "MovieSceneSequence.h"
-#include "Misc/QualifiedFrameTime.h"
 #include "Evaluation/MovieScenePlayback.h"
 
 UWorld* ULevelSequenceDirector::GetWorld() const
diff --git a/Engine/Source/Runtime/LevelSequence/Public/LevelSequenceDirector.h b/Engine/Source/Runtime/LevelSequence/Public/LevelSequenceDirector.h
--- a/Engine/Source/Runtime/LevelSequence/Public/LevelSequenceDirector.h
+++ b/Engine/Source/Runtime/LevelSequence/Public/LevelSequenceDirector.h
@@ -7,6 +7,7 @@
 #include "Engine/BlueprintGeneratedClass.h"
 #include "MovieSceneSequenceID.h"
 #include "MovieSceneObjectBindingID.h"
+#include "Misc/QualifiedFrameTime.h"
 #include "LevelSequenceDirector.generated.h"
 
 class IMovieScenePlayer;
