LevelSequence: Export symbols for FBoundActorProxy.

diff --git a/Engine/Source/Runtime/LevelSequence/Public/LevelSequenceActor.h b/Engine/Source/Runtime/LevelSequence/Public/LevelSequenceActor.h
--- a/Engine/Source/Runtime/LevelSequence/Public/LevelSequenceActor.h
+++ b/Engine/Source/Runtime/LevelSequence/Public/LevelSequenceActor.h
@@ -318,7 +318,7 @@ private:
 };
 
 USTRUCT()
-struct FBoundActorProxy
+struct LEVELSEQUENCE_API FBoundActorProxy
 {
 	GENERATED_BODY()
 
