Sequencer: Fix build

diff --git a/Engine/Source/Runtime/LevelSequence/Private/LevelSequence.cpp b/Engine/Source/Runtime/LevelSequence/Private/LevelSequence.cpp
--- a/Engine/Source/Runtime/LevelSequence/Private/LevelSequence.cpp
+++ b/Engine/Source/Runtime/LevelSequence/Private/LevelSequence.cpp
@@ -456,11 +456,12 @@ UObject* ULevelSequence::GetParentObject(UObject* Object) const
 
 bool ULevelSequence::AllowsSpawnableObjects() const
 {
+#if WITH_EDITOR
 	if (!UMovieScene::IsTrackClassAllowed(UMovieSceneSpawnTrack::StaticClass()))
 	{
 		return false;
 	}
-
+#endif
 	return true;
 }
 
