Sequencer: Replicate ALevelSequenceActor::LevelSequenceAsset and initialize the player if it changes

diff --git a/Engine/Source/Runtime/LevelSequence/Private/LevelSequenceActor.cpp b/Engine/Source/Runtime/LevelSequence/Private/LevelSequenceActor.cpp
--- a/Engine/Source/Runtime/LevelSequence/Private/LevelSequenceActor.cpp
+++ b/Engine/Source/Runtime/LevelSequence/Private/LevelSequenceActor.cpp
@@ -150,6 +150,7 @@ void ALevelSequenceActor::GetLifetimeReplicatedProps(TArray<FLifetimeProperty>&
 	Super::GetLifetimeReplicatedProps(OutLifetimeProps);
 
 	DOREPLIFETIME(ALevelSequenceActor, SequencePlayer);
+	DOREPLIFETIME(ALevelSequenceActor, LevelSequenceAsset);
 }
 
 void ALevelSequenceActor::PostInitializeComponents()
@@ -504,6 +505,16 @@ const TArray<FMovieSceneObjectBindingID>& ALevelSequenceActor::FindNamedBindings
 	return EmptyBindings;
 }
 
+void ALevelSequenceActor::PostNetReceive()
+{
+	Super::PostNetReceive();
+
+	if (LevelSequenceAsset && SequencePlayer && SequencePlayer->GetSequence() != LevelSequenceAsset)
+	{
+		InitializePlayer();
+	}
+}
+
 #if WITH_EDITOR
 
 void FBoundActorProxy::Initialize(TSharedPtr<IPropertyHandle> InPropertyHandle)
diff --git a/Engine/Source/Runtime/LevelSequence/Public/LevelSequenceActor.h b/Engine/Source/Runtime/LevelSequence/Public/LevelSequenceActor.h
--- a/Engine/Source/Runtime/LevelSequence/Public/LevelSequenceActor.h
+++ b/Engine/Source/Runtime/LevelSequence/Public/LevelSequenceActor.h
@@ -82,7 +82,7 @@ public:
 	UPROPERTY(Instanced, transient, replicated, BlueprintReadOnly, BlueprintGetter=GetSequencePlayer, Category="Playback", meta=(ExposeFunctionCategories="Sequencer|Player"))
 	TObjectPtr<ULevelSequencePlayer> SequencePlayer;
 
-	UPROPERTY(EditAnywhere, BlueprintReadOnly, Category="General", meta=(AllowedClasses="LevelSequence"))
+	UPROPERTY(EditAnywhere, replicated, BlueprintReadOnly, Category="General", meta=(AllowedClasses="LevelSequence"))
 	TObjectPtr<ULevelSequence> LevelSequenceAsset;
 
 #if WITH_EDITORONLY_DATA
@@ -281,6 +281,7 @@ protected:
 	virtual void BeginPlay() override;
 	virtual void EndPlay(const EEndPlayReason::Type EndPlayReason) override;
 	virtual void RewindForReplay() override;
+	virtual void PostNetReceive() override;
 #if WITH_EDITOR
 	virtual bool CanChangeIsSpatiallyLoadedFlag() const override { return false; }
 #endif
