ANY_PACKAGE Removal: ULevelSequence::PostLoadAssetRegistryTags needs to postload blueprint tags.
#rnx

diff --git a/Engine/Source/Runtime/LevelSequence/Private/LevelSequence.cpp b/Engine/Source/Runtime/LevelSequence/Private/LevelSequence.cpp
--- a/Engine/Source/Runtime/LevelSequence/Private/LevelSequence.cpp
+++ b/Engine/Source/Runtime/LevelSequence/Private/LevelSequence.cpp
@@ -170,6 +170,14 @@
 	Super::GetAssetRegistryTagMetadata(OutMetadata);
 }
 
+void ULevelSequence::PostLoadAssetRegistryTags(const FAssetData& InAssetData, TArray<FAssetRegistryTag>& OutTagsAndValuesToUpdate) const
+{
+	Super::PostLoadAssetRegistryTags(InAssetData, OutTagsAndValuesToUpdate);
+
+	// GetAssetRegistryTags appends the DirectorBlueprint tags to the World's tags, so we also have to run the Blueprint PostLoadAssetRegistryTags
+	UBlueprint::PostLoadBlueprintAssetRegistryTags(InAssetData, OutTagsAndValuesToUpdate);
+}
+
 void PurgeLegacyBlueprints(UObject* InObject, UPackage* Package)
 {
 	if (UBlueprint* BP = Cast<UBlueprint>(InObject))
diff --git a/Engine/Source/Runtime/LevelSequence/Public/LevelSequence.h b/Engine/Source/Runtime/LevelSequence/Public/LevelSequence.h
--- a/Engine/Source/Runtime/LevelSequence/Public/LevelSequence.h
+++ b/Engine/Source/Runtime/LevelSequence/Public/LevelSequence.h
@@ -75,6 +75,7 @@ public:
 	virtual ETrackSupport IsTrackSupported(TSubclassOf<class UMovieSceneTrack> InTrackClass) const override;
 	virtual void GetAssetRegistryTagMetadata(TMap<FName, FAssetRegistryTagMetadata>& OutMetadata) const override;
 	virtual void GetAssetRegistryTags(TArray<FAssetRegistryTag>& OutTags) const override;
+	virtual void PostLoadAssetRegistryTags(const FAssetData& InAssetData, TArray<FAssetRegistryTag>& OutTagsAndValuesToUpdate) const override;
 	
 	DECLARE_DELEGATE_RetVal_OneParam(void, FPostDuplicateEvent, ULevelSequence*);
 	static FPostDuplicateEvent PostDuplicateEvent;
