Sequencer: Add guards when IsSavingPackage or IsGarbageCollecting before calling FindObject

Illegal call to StaticFindObjectFast() while serializing object data or garbage collecting.

diff --git a/Engine/Source/Runtime/LevelSequence/Private/LevelSequenceBindingReference.cpp b/Engine/Source/Runtime/LevelSequence/Private/LevelSequenceBindingReference.cpp
--- a/Engine/Source/Runtime/LevelSequence/Private/LevelSequenceBindingReference.cpp
+++ b/Engine/Source/Runtime/LevelSequence/Private/LevelSequenceBindingReference.cpp
@@ -2,8 +2,10 @@
 
 #include "LevelSequenceBindingReference.h"
 #include "LevelSequenceLegacyObjectReference.h"
+#include "UObject/GarbageCollection.h"
 #include "UObject/Package.h"
 #include "UObject/ObjectMacros.h"
+#include "UObject/UObjectGlobals.h"
 #include "MovieSceneFwd.h"
 #include "Misc/PackageName.h"
 #include "Engine/World.h"
@@ -47,11 +49,21 @@ UObject* FLevelSequenceBindingReference::Resolve(UObject* InContext, FName Strea
 	{
 		if (ExternalObjectPath.IsNull())
 		{
+			if (UE::IsSavingPackage(nullptr) || IsGarbageCollecting())
+			{
+				return nullptr;
+			}
+
 			return FindObject<UObject>(InContext, *ObjectPath, false);
 		}
 	}
 	else if (InContext && InContext->IsA<ULevel>() && StreamedLevelAssetPath != NAME_None && ExternalObjectPath.GetAssetPathName() == StreamedLevelAssetPath)
 	{
+		if (UE::IsSavingPackage(nullptr) || IsGarbageCollecting())
+		{
+			return nullptr;
+		}
+
 		// ExternalObjectPath.GetSubPathString() specifies the path from the package (so includes PersistentLevel.) so we must do a FindObject from its outer
 		return FindObject<UObject>(InContext->GetOuter(), *ExternalObjectPath.GetSubPathString());
 	}
@@ -105,6 +117,11 @@ void FLevelSequenceBindingReference::PostSerialize(const FArchive& Ar)
 
 UObject* ResolveByPath(UObject* InContext, const FString& InObjectPath)
 {
+	if (UE::IsSavingPackage(nullptr) || IsGarbageCollecting())
+	{
+		return nullptr;
+	}
+
 	if (!InObjectPath.IsEmpty())
 	{
 		if (UObject* FoundObject = FindObject<UObject>(InContext, *InObjectPath, false))
