Sequencer: Only ever call FixupForPIE for bindings that are being resolved within a World context that has a PIE instance ID
Sometimes binding resolution can be triggered inside a call stack that has a GPlayInEditorID, even if the bindings and resolution context do not. In these situations we were previously erroneously calling FixupForPIE before resolving the soft object path, and ending up finding the PIE actor instead of the one that's actually inside the world context provided. We now explicitly override the GPlayInEditorID for bindings based on the provided context object.

diff --git a/Engine/Source/Runtime/LevelSequence/Private/LevelSequenceBindingReference.cpp b/Engine/Source/Runtime/LevelSequence/Private/LevelSequenceBindingReference.cpp
--- a/Engine/Source/Runtime/LevelSequence/Private/LevelSequenceBindingReference.cpp
+++ b/Engine/Source/Runtime/LevelSequence/Private/LevelSequenceBindingReference.cpp
@@ -77,18 +77,12 @@ UObject* FLevelSequenceBindingReference::Resolve(UObject* InContext, FName Strea
 		TempPath.PreSavePath();
 
 	#if WITH_EDITORONLY_DATA
-		int32 ContextPlayInEditorID = InContext ? InContext->GetOutermost()->GetPIEInstanceID() : INDEX_NONE;
-
-		if (ContextPlayInEditorID != INDEX_NONE)
-		{
-			// We have an override PIE id, so set the global before entering
-			TGuardValue<int32> PIEGuard(GPlayInEditorID, ContextPlayInEditorID);
-			TempPath.FixupForPIE();
-		}
-		else
-		{
-			TempPath.FixupForPIE();
-		}
+		// Sequencer is explicit about providing a resolution context for its bindings. We never want to resolve to objects
+		// with a different PIE instance ID, even if the current callstack is being executed inside a different GPlayInEditorID
+		// scope. Since ResolveObject will always call FixupForPIE in editor based on GPlayInEditorID, we always override the current
+		// GPlayInEditorID to be the current PIE instance of the provided context.
+		const int32 ContextPlayInEditorID = InContext ? InContext->GetOutermost()->GetPIEInstanceID() : INDEX_NONE;
+		TGuardValue<int32> PIEGuard(GPlayInEditorID, ContextPlayInEditorID);
 	#endif
 
 		return TempPath.ResolveObject();
