Sequencer: Replace spaces with underscores in spawned actor names

diff --git a/Engine/Source/Runtime/LevelSequence/Private/LevelSequenceActorSpawner.cpp b/Engine/Source/Runtime/LevelSequence/Private/LevelSequenceActorSpawner.cpp
--- a/Engine/Source/Runtime/LevelSequence/Private/LevelSequenceActorSpawner.cpp
+++ b/Engine/Source/Runtime/LevelSequence/Private/LevelSequenceActorSpawner.cpp
@@ -139,7 +139,7 @@ UObject* FLevelSequenceActorSpawner::SpawnObject(FMovieSceneSpawnable& Spawnable
 	// Spawn the puppet actor
 	FActorSpawnParameters SpawnInfo;
 	{
-		SpawnInfo.Name = SpawnName;
+		SpawnInfo.Name = *SpawnName.ToString().Replace(TEXT(" "), TEXT("_"));
 		SpawnInfo.ObjectFlags = ObjectFlags;
 		SpawnInfo.SpawnCollisionHandlingOverride = ESpawnActorCollisionHandlingMethod::AlwaysSpawn;
 		// @todo: Spawning with a non-CDO template is fraught with issues
