Sequencer: Fix cinematic sequence device's level sequence actor spawn transform, and add option for always-on networking

diff --git a/Engine/Source/Runtime/LevelSequence/Private/LevelSequenceActor.cpp b/Engine/Source/Runtime/LevelSequence/Private/LevelSequenceActor.cpp
--- a/Engine/Source/Runtime/LevelSequence/Private/LevelSequenceActor.cpp
+++ b/Engine/Source/Runtime/LevelSequence/Private/LevelSequenceActor.cpp
@@ -612,3 +612,8 @@ void ULevelSequenceBurnInOptions::PostEditChangeProperty( FPropertyChangedEvent&
 
 #endif // WITH_EDITOR
 
+AReplicatedLevelSequenceActor::AReplicatedLevelSequenceActor(const FObjectInitializer& Init)
+	: Super(Init)
+{
+	bAlwaysRelevant = true;
+}
diff --git a/Engine/Source/Runtime/LevelSequence/Public/LevelSequenceActor.h b/Engine/Source/Runtime/LevelSequence/Public/LevelSequenceActor.h
--- a/Engine/Source/Runtime/LevelSequence/Public/LevelSequenceActor.h
+++ b/Engine/Source/Runtime/LevelSequence/Public/LevelSequenceActor.h
@@ -331,3 +331,17 @@ struct LEVELSEQUENCE_API FBoundActorProxy
 
 #endif
 };
+
+/**
+ * A level sequence actor that is set to always be relevant for networking purposes
+ */
+UCLASS()
+class LEVELSEQUENCE_API AReplicatedLevelSequenceActor
+	: public ALevelSequenceActor
+{
+	GENERATED_BODY()
+
+public:
+	/** Create and initialize a new instance. */
+	AReplicatedLevelSequenceActor(const FObjectInitializer& Init);
+};
