Updating Transform Origin tool-tip to remove warning about attach tracks.

diff --git a/Engine/Source/Runtime/LevelSequence/Public/DefaultLevelSequenceInstanceData.h b/Engine/Source/Runtime/LevelSequence/Public/DefaultLevelSequenceInstanceData.h
--- a/Engine/Source/Runtime/LevelSequence/Public/DefaultLevelSequenceInstanceData.h
+++ b/Engine/Source/Runtime/LevelSequence/Public/DefaultLevelSequenceInstanceData.h
@@ -32,7 +32,7 @@ public:
 	UPROPERTY(EditAnywhere, BlueprintReadWrite, Category="General")
 	TObjectPtr<AActor> TransformOriginActor;
 
-	/** Specifies a transform that offsets all absolute transform sections in this sequence. Will compound with attach tracks. Scale is ignored. Not applied to Relative or Additive sections. */
+	/** Specifies a transform that offsets all absolute transform sections in this sequence. Scale is ignored. Not applied to Relative or Additive sections. */
 	UPROPERTY(EditAnywhere, BlueprintReadWrite, Category="General")
 	FTransform TransformOrigin;
 };
