Sequencer: Spawn with a unique name

diff --git a/Engine/Source/Runtime/LevelSequence/Private/LevelSequenceActorSpawner.cpp b/Engine/Source/Runtime/LevelSequence/Private/LevelSequenceActorSpawner.cpp
--- a/Engine/Source/Runtime/LevelSequence/Private/LevelSequenceActorSpawner.cpp
+++ b/Engine/Source/Runtime/LevelSequence/Private/LevelSequenceActorSpawner.cpp
@@ -144,6 +144,7 @@ UObject* FLevelSequenceActorSpawner::SpawnObject(FMovieSceneSpawnable& Spawnable
 	FActorSpawnParameters SpawnInfo;
 	{
 		SpawnInfo.Name = *SpawnName.ToString().Replace(TEXT(" "), TEXT("_"));
+		SpawnInfo.Name = MakeUniqueObjectName(WorldContext->PersistentLevel, ObjectTemplate->GetClass(), SpawnInfo.Name);
 		SpawnInfo.ObjectFlags = ObjectFlags;
 		SpawnInfo.SpawnCollisionHandlingOverride = ESpawnActorCollisionHandlingMethod::AlwaysSpawn;
 		// @todo: Spawning with a non-CDO template is fraught with issues
