Sequencer: Fixed Level Sequence blueprint having invalid characters in its name

diff --git a/Engine/Source/Runtime/LevelSequence/Private/LevelSequence.cpp b/Engine/Source/Runtime/LevelSequence/Private/LevelSequence.cpp
--- a/Engine/Source/Runtime/LevelSequence/Private/LevelSequence.cpp
+++ b/Engine/Source/Runtime/LevelSequence/Private/LevelSequence.cpp
@@ -518,7 +518,7 @@ UBlueprint* ULevelSequence::GetDirectorBlueprint() const
 
 FString ULevelSequence::GetDirectorBlueprintName() const
 {
-	return GetDisplayName().ToString() + " (Director BP)";
+	return GetDisplayName().ToString() + "_DirectorBP";
 }
 
 void ULevelSequence::SetDirectorBlueprint(UBlueprint* NewDirectorBlueprint)
