Sequencer: Fix level sequence project settings not persisting because they are tied to cvars.

The bug is that when you reload the engine, the config is properly restored, but are then overwritten by the corresponding cvar value. That means if you edit the setting through the UI or set it through ini, the default value from the cvar will win out. This is not desirable.

Tested the following:
- If you edit the project setting, restart, and the project setting will persist. A created level sequence will have that same setting.
- If you edit the ini, restart, and the project setting will persist. A created level sequence will have that same setting.
- If you change the cvar (either on load or after the engine has loaded), the project setting will not be updated. Any created level sequence during this session will have that setting.

diff --git a/Engine/Source/Runtime/LevelSequence/Private/LevelSequenceProjectSettings.cpp b/Engine/Source/Runtime/LevelSequence/Private/LevelSequenceProjectSettings.cpp
--- a/Engine/Source/Runtime/LevelSequence/Private/LevelSequenceProjectSettings.cpp
+++ b/Engine/Source/Runtime/LevelSequence/Private/LevelSequenceProjectSettings.cpp
@@ -2,6 +2,8 @@
 
 #include "LevelSequenceProjectSettings.h"
 
+#include "HAL/IConsoleManager.h"
+
 #include UE_INLINE_GENERATED_CPP_BY_NAME(LevelSequenceProjectSettings)
 
 ULevelSequenceProjectSettings::ULevelSequenceProjectSettings()
@@ -19,7 +21,31 @@ void ULevelSequenceProjectSettings::PostInitProperties()
 #if WITH_EDITOR
 	if(IsTemplate())
 	{
-		ImportConsoleVariableValues();
+		// Most classes that uses this console backed cvar to import the cvar values onto the config. 
+		// This isn't quite desirable because it means changing a value in the project settings will 
+		// get overwritten by cvar (possibly default) values. So, instead of importing console variable 
+		// values onto the config, set the console variable values based on the config values.
+		// ImportConsoleVariableValues();
+	
+		if (IConsoleVariable* ConsoleVariable = IConsoleManager::Get().FindConsoleVariable(TEXT("LevelSequence.DefaultLockEngineToDisplayRate")))
+		{
+			ConsoleVariable->Set((int32)bDefaultLockEngineToDisplayRate, ECVF_SetByProjectSetting);
+		}
+
+		if (IConsoleVariable* ConsoleVariable = IConsoleManager::Get().FindConsoleVariable(TEXT("LevelSequence.DefaultTickResolution")))
+		{
+			ConsoleVariable->Set(*DefaultTickResolution, ECVF_SetByProjectSetting);
+		}
+
+		if (IConsoleVariable* ConsoleVariable = IConsoleManager::Get().FindConsoleVariable(TEXT("LevelSequence.DefaultDisplayRate")))
+		{
+			ConsoleVariable->Set(*DefaultDisplayRate, ECVF_SetByProjectSetting);
+		}
+
+		if (IConsoleVariable* ConsoleVariable = IConsoleManager::Get().FindConsoleVariable(TEXT("LevelSequence.DefaultClockSource")))
+		{
+			ConsoleVariable->Set((int32)DefaultClockSource, ECVF_SetByProjectSetting);
+		}
 	}
 #endif
 }
