Sequencer: Add SetSpawnableBindingID which exposes binding to a spawnable from blueprints/python.

Remove restriction on being able to bind to a spawnable in LevelSequence::FindOrAddBinding since that is now allowed

diff --git a/Engine/Plugins/MovieScene/SequencerScripting/Source/SequencerScripting/Private/ExtensionLibraries/MovieSceneBindingExtensions.cpp b/Engine/Plugins/MovieScene/SequencerScripting/Source/SequencerScripting/Private/ExtensionLibraries/MovieSceneBindingExtensions.cpp
--- a/Engine/Plugins/MovieScene/SequencerScripting/Source/SequencerScripting/Private/ExtensionLibraries/MovieSceneBindingExtensions.cpp
+++ b/Engine/Plugins/MovieScene/SequencerScripting/Source/SequencerScripting/Private/ExtensionLibraries/MovieSceneBindingExtensions.cpp
@@ -343,3 +343,20 @@ void UMovieSceneBindingExtensions::MoveBindingContents(const FMovieSceneBindingP
 	}
 }
 
+void UMovieSceneBindingExtensions::SetSpawnableBindingID(const FMovieSceneBindingProxy& InBinding, const FMovieSceneObjectBindingID& SpawnableBindingID)
+{
+	if (!InBinding.Sequence)
+	{
+		FFrame::KismetExecutionMessage(TEXT("Cannot call SetSpawnableBindingID with a binding with a null sequence"), ELogVerbosity::Error);
+		return;
+	}
+
+	FMovieScenePossessable* Possessable = InBinding.Sequence->GetMovieScene()->FindPossessable(InBinding.BindingID);
+	if (!Possessable)
+	{
+		FFrame::KismetExecutionMessage(TEXT("Cannot call SetSpawnableBindingID with a null possessable"), ELogVerbosity::Error);
+		return;
+	}
+	
+	Possessable->SetSpawnableObjectBindingID(SpawnableBindingID);
+}
diff --git a/Engine/Plugins/MovieScene/SequencerScripting/Source/SequencerScripting/Public/ExtensionLibraries/MovieSceneBindingExtensions.h b/Engine/Plugins/MovieScene/SequencerScripting/Source/SequencerScripting/Public/ExtensionLibraries/MovieSceneBindingExtensions.h
--- a/Engine/Plugins/MovieScene/SequencerScripting/Source/SequencerScripting/Public/ExtensionLibraries/MovieSceneBindingExtensions.h
+++ b/Engine/Plugins/MovieScene/SequencerScripting/Source/SequencerScripting/Public/ExtensionLibraries/MovieSceneBindingExtensions.h
@@ -197,6 +197,15 @@ public:
 	 */
 	UFUNCTION(BlueprintCallable, Category = "Sequencer|Sequence", meta = (ScriptMethod))
 	static void MoveBindingContents(const FMovieSceneBindingProxy& SourceBindingId, const FMovieSceneBindingProxy& DestinationBindingId);
+
+	/**
+	 * Set the spawnable id that the possessable binding should possess
+	 *
+	 * @param InBinding     The binding to set
+	 * @param SpawnableBindingID The spawnable's binding id
+	 */
+	UFUNCTION(BlueprintCallable, Category = "Sequencer|Sequence", meta = (ScriptMethod))
+	static void SetSpawnableBindingID(const FMovieSceneBindingProxy& InBinding, const FMovieSceneObjectBindingID& SpawnableBindingID);
 };
 
 #if UE_ENABLE_INCLUDE_ORDER_DEPRECATED_IN_5_2
diff --git a/Engine/Source/Runtime/LevelSequence/Private/LevelSequence.cpp b/Engine/Source/Runtime/LevelSequence/Private/LevelSequence.cpp
--- a/Engine/Source/Runtime/LevelSequence/Private/LevelSequence.cpp
+++ b/Engine/Source/Runtime/LevelSequence/Private/LevelSequence.cpp
@@ -560,9 +560,6 @@ FGuid ULevelSequence::FindOrAddBinding(UObject* InObject)
 	}
 
 	AActor* Actor = Cast<AActor>(InObject);
-	// @todo: sequencer-python: need to figure out how we go from a spawned object to an object binding without the spawn register or any IMovieScenePlayer interface
-	// Normally this process would happen through sequencer, since it has more context than just the level sequence asset.
-	// For now we cannot possess spawnables or anything within them since we have no way of retrieving the spawnable from the object
 	if (Actor && Actor->ActorHasTag("SequencerActor"))
 	{
 		TOptional<FMovieSceneSpawnableAnnotation> Annotation = FMovieSceneSpawnableAnnotation::Find(Actor);
@@ -570,9 +567,9 @@ FGuid ULevelSequence::FindOrAddBinding(UObject* InObject)
 		{
 			return Annotation->ObjectBindingID;
 		}
-
-		UE_LOG(LogLevelSequence, Error, TEXT("Unable to possess object '%s' since it is, or is part of a spawnable that is not in this sequence."), *InObject->GetName());
-		return FGuid();
+		
+		// If this actor is a spawnable and is not in the same originating sequence, it's likely a spawnable that will be possessed. 
+		// SetSpawnableObjectBindingID will need to be called on that possessable.
 	}
 
 	UObject* ParentObject = GetParentObject(InObject);
diff --git a/Engine/Source/Runtime/MovieScene/Public/MovieScenePossessable.h b/Engine/Source/Runtime/MovieScene/Public/MovieScenePossessable.h
--- a/Engine/Source/Runtime/MovieScene/Public/MovieScenePossessable.h
+++ b/Engine/Source/Runtime/MovieScene/Public/MovieScenePossessable.h
@@ -154,7 +154,7 @@ public:
 	}
 
 	/* Set the optional binding id for binding to a spawnable */
-	void SetSpawnableObjectBindingID(const FMovieSceneObjectBindingID& InSpawnableObjectBindingID)
+	MOVIESCENE_API void SetSpawnableObjectBindingID(const FMovieSceneObjectBindingID& InSpawnableObjectBindingID)
 	{
 		SpawnableObjectBindingID = InSpawnableObjectBindingID;
 	}
