Sequencer: take play-rate into account for blend times between blending camera cuts

diff --git a/Engine/Source/Runtime/LevelSequence/Private/LevelSequencePlayer.cpp b/Engine/Source/Runtime/LevelSequence/Private/LevelSequencePlayer.cpp
--- a/Engine/Source/Runtime/LevelSequence/Private/LevelSequencePlayer.cpp
+++ b/Engine/Source/Runtime/LevelSequence/Private/LevelSequencePlayer.cpp
@@ -351,9 +351,11 @@ void ULevelSequencePlayer::UpdateCameraCut(UObject* CameraObject, const EMovieSc
 			(CameraObject ? *CameraObject->GetName() : TEXT("None")),
 			TransitionParams.BlendTime);
 
+		const float PlayRateFactor = PlaybackSettings.PlayRate != 0.f ? (1.f / PlaybackSettings.PlayRate) : 1.f;
+
 		// Convert known easing functions to their corresponding view target blend parameters.
 		TTuple<EViewTargetBlendFunction, float> BlendFunctionAndExp = BuiltInEasingTypeToBlendFunction(CameraCutParams.BlendType.GetValue());
-		TransitionParams.BlendTime = CameraCutParams.BlendTime;
+		TransitionParams.BlendTime = CameraCutParams.BlendTime * PlayRateFactor;
 		TransitionParams.bLockOutgoing = CameraCutParams.bLockPreviousCamera;
 		TransitionParams.BlendFunction = BlendFunctionAndExp.Get<0>();
 		TransitionParams.BlendExp = BlendFunctionAndExp.Get<1>();
