Keep level sequence player up to date with view target changes during playback, so that it restores the last known desired view target when stopping

#localization none
#tests PIE

diff --git a/Engine/Source/Runtime/LevelSequence/Private/LevelSequencePlayer.cpp b/Engine/Source/Runtime/LevelSequence/Private/LevelSequencePlayer.cpp
--- a/Engine/Source/Runtime/LevelSequence/Private/LevelSequencePlayer.cpp
+++ b/Engine/Source/Runtime/LevelSequence/Private/LevelSequencePlayer.cpp
@@ -171,6 +171,7 @@ void ULevelSequencePlayer::OnStopped()
 		}
 	}
 
+	LastCameraObject = nullptr;
 	LastViewTarget.Reset();
 }
 
@@ -330,7 +331,7 @@ void ULevelSequencePlayer::UpdateCameraCut(UObject* CameraObject, const EMovieSc
 	}
 
 	// Save the last view target/aspect ratio constraint/etc. so that it can all be restored when the camera object is null.
-	if (!LastViewTarget.IsValid())
+	if (!LastViewTarget.IsValid() || (ViewTarget != CameraObject && ViewTarget != LastCameraObject))
 	{
 		LastViewTarget = ViewTarget;
 	}
@@ -342,6 +343,8 @@ void ULevelSequencePlayer::UpdateCameraCut(UObject* CameraObject, const EMovieSc
 		}
 	}
 
+	LastCameraObject = CameraObject;
+
 	bool bDoSetViewTarget = true;
 	FViewTargetTransitionParams TransitionParams;
 	if (CameraCutParams.BlendType.IsSet())
diff --git a/Engine/Source/Runtime/LevelSequence/Public/LevelSequencePlayer.h b/Engine/Source/Runtime/LevelSequence/Public/LevelSequencePlayer.h
--- a/Engine/Source/Runtime/LevelSequence/Public/LevelSequencePlayer.h
+++ b/Engine/Source/Runtime/LevelSequence/Public/LevelSequencePlayer.h
@@ -185,6 +185,9 @@ private:
 	/** The last aspect ratio axis constraint to reset to when the camera cut is null */
 	TOptional<EAspectRatioAxisConstraint> LastAspectRatioAxisConstraint;
 
+	/** The last camera cut object received during a cinematic */
+	UObject* LastCameraObject = nullptr;
+
 protected:
 
 	TOptional<int32> SnapshotOffsetTime;
