Sequencer: Reinitialize player if the runner is no longer valid

This can happen if a level is unloaded and reloaded without a garbage collection pass happening in-between, meaning the player would be torn down but wouldn't get re-created since the player was already partly initialized.

diff --git a/Engine/Source/Runtime/LevelSequence/Private/LevelSequenceActor.cpp b/Engine/Source/Runtime/LevelSequence/Private/LevelSequenceActor.cpp
--- a/Engine/Source/Runtime/LevelSequence/Private/LevelSequenceActor.cpp
+++ b/Engine/Source/Runtime/LevelSequence/Private/LevelSequenceActor.cpp
@@ -278,7 +278,7 @@ void ALevelSequenceActor::InitializePlayer()
 	if (LevelSequenceAsset && GetWorld()->IsGameWorld())
 	{
 		// Level sequence is already loaded. Initialize the player if it's not already initialized with this sequence
-		if (LevelSequenceAsset != SequencePlayer->GetSequence())
+		if (LevelSequenceAsset != SequencePlayer->GetSequence() || SequencePlayer->GetEvaluationTemplate().GetRunner() == nullptr)
 		{
 			SequencePlayer->Initialize(LevelSequenceAsset, GetLevel(), CameraSettings);
 		}
