Fix non unity build

#rnx

diff --git a/Engine/Source/Runtime/LevelSequence/Private/LevelSequenceBindingReference.cpp b/Engine/Source/Runtime/LevelSequence/Private/LevelSequenceBindingReference.cpp
--- a/Engine/Source/Runtime/LevelSequence/Private/LevelSequenceBindingReference.cpp
+++ b/Engine/Source/Runtime/LevelSequence/Private/LevelSequenceBindingReference.cpp
@@ -7,6 +7,7 @@
 #include "UObject/Package.h"
 #include "UObject/ObjectMacros.h"
 #include "UObject/UObjectGlobals.h"
+#include "UnrealEngine.h"
 #include "MovieSceneFwd.h"
 #include "Misc/PackageName.h"
 #include "Engine/World.h"
