Sequencer: restore aspect ratio axis constraint directly on the player if the level has changed during sequence playback

#udn 00420602

diff --git a/Engine/Source/Runtime/LevelSequence/Private/LevelSequencePlayer.cpp b/Engine/Source/Runtime/LevelSequence/Private/LevelSequencePlayer.cpp
--- a/Engine/Source/Runtime/LevelSequence/Private/LevelSequencePlayer.cpp
+++ b/Engine/Source/Runtime/LevelSequence/Private/LevelSequencePlayer.cpp
@@ -173,6 +173,8 @@ void ULevelSequencePlayer::OnStopped()
 
 	LastCameraObject = nullptr;
 	LastViewTarget.Reset();
+	LastLocalPlayer.Reset();
+	LastAspectRatioAxisConstraint.Reset();
 }
 
 void ULevelSequencePlayer::UpdateMovieSceneInstance(FMovieSceneEvaluationRange InRange, EMovieScenePlayerStatus::Type PlayerStatus, const FMovieSceneUpdateArgs& Args)
@@ -266,6 +268,17 @@ void ULevelSequencePlayer::UpdateCameraCut(UObject* CameraObject, const EMovieSc
 
 	if (PC == nullptr)
 	{
+		// If we are releasing view target control but can't actually do it because game state has been
+		// torn down, see if we can at least restore aspect ratio axis constraint. This fixes a bug
+		// where the wrong constraint carries over to a new level if the level change was executed
+		// while we were playing the sequence.
+		if (CameraObject == nullptr && LastAspectRatioAxisConstraint.IsSet())
+		{
+			if (ULocalPlayer* LocalPlayer = LastLocalPlayer.Get())
+			{
+				LocalPlayer->AspectRatioAxisConstraint = LastAspectRatioAxisConstraint.GetValue();
+			}
+		}
 		return;
 	}
 
@@ -335,6 +348,10 @@ void ULevelSequencePlayer::UpdateCameraCut(UObject* CameraObject, const EMovieSc
 	{
 		LastViewTarget = ViewTarget;
 	}
+	if (!LastLocalPlayer.IsValid() || (LocalPlayer != LastLocalPlayer))
+	{
+		LastLocalPlayer = LocalPlayer;
+	}
 	if (!LastAspectRatioAxisConstraint.IsSet())
 	{
 		if (LocalPlayer != nullptr)
diff --git a/Engine/Source/Runtime/LevelSequence/Public/LevelSequencePlayer.h b/Engine/Source/Runtime/LevelSequence/Public/LevelSequencePlayer.h
--- a/Engine/Source/Runtime/LevelSequence/Public/LevelSequencePlayer.h
+++ b/Engine/Source/Runtime/LevelSequence/Public/LevelSequencePlayer.h
@@ -178,6 +178,9 @@ private:
 	/** The last view target to reset to when updating camera cuts to null */
 	TWeakObjectPtr<AActor> LastViewTarget;
 
+	/** The last player on which to reset the aspect ratio axis constraint in case the level is changing */
+	TWeakObjectPtr<ULocalPlayer> LastLocalPlayer;
+
 	/** The last aspect ratio axis constraint to reset to when the camera cut is null */
 	TOptional<EAspectRatioAxisConstraint> LastAspectRatioAxisConstraint;
 
