Sequencer: Follow up to CL 25176010 Comments for shots/sections

On PostDuplicate, reinitialize the metadata with owner, date, and clear the notes

diff --git a/Engine/Source/Runtime/LevelSequence/Private/LevelSequence.cpp b/Engine/Source/Runtime/LevelSequence/Private/LevelSequence.cpp
--- a/Engine/Source/Runtime/LevelSequence/Private/LevelSequence.cpp
+++ b/Engine/Source/Runtime/LevelSequence/Private/LevelSequence.cpp
@@ -273,6 +273,13 @@ void ULevelSequence::PostDuplicate(bool bDuplicateForPIE)
 		PostDuplicateEvent.Execute(this);
 	}
 #endif
+
+#if WITH_EDITOR
+	UMovieSceneMetaData* MetaData = FindOrAddMetaData<UMovieSceneMetaData>();
+	MetaData->SetCreated(FDateTime::UtcNow());
+	MetaData->SetAuthor(FApp::GetSessionOwner());
+	MetaData->SetNotes(FString()); // Intentionally clear the notes
+#endif
 }
 
 void ULevelSequence::PostLoad()
