LevelSequence: Replace single StreamedLevelAssetPath param by a new struct pointing to StreamedLevelAssetPath (Ground work for new resolve path)

#rnx

diff --git a/Engine/Source/Runtime/LevelSequence/Private/LevelSequence.cpp b/Engine/Source/Runtime/LevelSequence/Private/LevelSequence.cpp
--- a/Engine/Source/Runtime/LevelSequence/Private/LevelSequence.cpp
+++ b/Engine/Source/Runtime/LevelSequence/Private/LevelSequence.cpp
@@ -433,10 +433,10 @@ bool ULevelSequence::CanPossessObject(UObject& Object, UObject* InPlaybackContex
 
 void ULevelSequence::LocateBoundObjects(const FGuid& ObjectId, UObject* Context, TArray<UObject*, TInlineAllocator<1>>& OutObjects) const
 {
-	LocateBoundObjects(ObjectId, Context, {}, OutObjects);
+	LocateBoundObjects(ObjectId, Context, FLevelSequenceBindingReference::FResolveBindingParams(), OutObjects);
 }
 
-void ULevelSequence::LocateBoundObjects(const FGuid& ObjectId, UObject* Context, const FTopLevelAssetPath& StreamedLevelAssetPath, TArray<UObject*, TInlineAllocator<1>>& OutObjects) const
+void ULevelSequence::LocateBoundObjects(const FGuid& ObjectId, UObject* Context, const FLevelSequenceBindingReference::FResolveBindingParams& InResolveBindingParams, TArray<UObject*, TInlineAllocator<1>>& OutObjects) const
 {
 	// Handle legacy object references
 	UObject* Object = Context ? ObjectReferences.ResolveBinding(ObjectId, Context) : nullptr;
@@ -445,7 +445,7 @@ void ULevelSequence::LocateBoundObjects(const FGuid& ObjectId, UObject* Context,
 		OutObjects.Add(Object);
 	}
 
-	BindingReferences.ResolveBinding(ObjectId, Context, StreamedLevelAssetPath, OutObjects);
+	BindingReferences.ResolveBinding(ObjectId, Context, InResolveBindingParams, OutObjects);
 }
 
 FGuid ULevelSequence::FindBindingFromObject(UObject* InObject, UObject* Context) const
diff --git a/Engine/Source/Runtime/LevelSequence/Private/LevelSequenceBindingReference.cpp b/Engine/Source/Runtime/LevelSequence/Private/LevelSequenceBindingReference.cpp
--- a/Engine/Source/Runtime/LevelSequence/Private/LevelSequenceBindingReference.cpp
+++ b/Engine/Source/Runtime/LevelSequence/Private/LevelSequenceBindingReference.cpp
@@ -47,7 +47,7 @@ FLevelSequenceBindingReference::FLevelSequenceBindingReference(UObject* InObject
 	}
 }
 
-UObject* FLevelSequenceBindingReference::Resolve(UObject* InContext, const FTopLevelAssetPath& StreamedLevelAssetPath) const
+UObject* FLevelSequenceBindingReference::Resolve(UObject* InContext, const FResolveBindingParams& InResolveBindingParams) const
 {
 	if (InContext && InContext->IsA<AActor>())
 	{
@@ -61,7 +61,7 @@ UObject* FLevelSequenceBindingReference::Resolve(UObject* InContext, const FTopL
 			return FindObject<UObject>(InContext, *ObjectPath, false);
 		}
 	}
-	else if (InContext && InContext->IsA<ULevel>() && StreamedLevelAssetPath.IsValid() && ExternalObjectPath.GetAssetPath() == StreamedLevelAssetPath)
+	else if (InContext && InContext->IsA<ULevel>() && InResolveBindingParams.StreamedLevelAssetPath.IsValid() && ExternalObjectPath.GetAssetPath() == InResolveBindingParams.StreamedLevelAssetPath)
 	{
 		if (UE::IsSavingPackage(nullptr) || IsGarbageCollecting())
 		{
@@ -266,7 +266,7 @@ void FLevelSequenceBindingReferences::RemoveObjects(const FGuid& ObjectId, const
 
 	for (int32 ReferenceIndex = 0; ReferenceIndex < ReferenceArray->References.Num(); )
 	{
-		UObject* ResolvedObject = ReferenceArray->References[ReferenceIndex].Resolve(InContext, FTopLevelAssetPath());
+		UObject* ResolvedObject = ReferenceArray->References[ReferenceIndex].Resolve(InContext, FLevelSequenceBindingReference::FResolveBindingParams());
 
 		if (InObjects.Contains(ResolvedObject))
 		{
@@ -289,7 +289,7 @@ void FLevelSequenceBindingReferences::RemoveInvalidObjects(const FGuid& ObjectId
 
 	for (int32 ReferenceIndex = 0; ReferenceIndex < ReferenceArray->References.Num(); )
 	{
-		UObject* ResolvedObject = ReferenceArray->References[ReferenceIndex].Resolve(InContext, FTopLevelAssetPath());
+		UObject* ResolvedObject = ReferenceArray->References[ReferenceIndex].Resolve(InContext, FLevelSequenceBindingReference::FResolveBindingParams());
 
 		if (!IsValid(ResolvedObject))
 		{
@@ -302,13 +302,13 @@ void FLevelSequenceBindingReferences::RemoveInvalidObjects(const FGuid& ObjectId
 	}
 }
 
-void FLevelSequenceBindingReferences::ResolveBinding(const FGuid& ObjectId, UObject* InContext, const FTopLevelAssetPath& StreamedLevelAssetPath, TArray<UObject*, TInlineAllocator<1>>& OutObjects) const
+void FLevelSequenceBindingReferences::ResolveBinding(const FGuid& ObjectId, UObject* InContext, const FLevelSequenceBindingReference::FResolveBindingParams& InResolveBindingParams, TArray<UObject*, TInlineAllocator<1>>& OutObjects) const
 {
 	if (const FLevelSequenceBindingReferenceArray* ReferenceArray = BindingIdToReferences.Find(ObjectId))
 	{
 		for (const FLevelSequenceBindingReference& Reference : ReferenceArray->References)
 		{
-			UObject* ResolvedObject = Reference.Resolve(InContext, StreamedLevelAssetPath);
+			UObject* ResolvedObject = Reference.Resolve(InContext, InResolveBindingParams);
 			if (ResolvedObject && ResolvedObject->GetWorld())
 			{
 				OutObjects.Add(ResolvedObject);
diff --git a/Engine/Source/Runtime/LevelSequence/Private/LevelSequencePlayer.cpp b/Engine/Source/Runtime/LevelSequence/Private/LevelSequencePlayer.cpp
--- a/Engine/Source/Runtime/LevelSequence/Private/LevelSequencePlayer.cpp
+++ b/Engine/Source/Runtime/LevelSequence/Private/LevelSequencePlayer.cpp
@@ -134,8 +134,9 @@ void ULevelSequencePlayer::ResolveBoundObjects(const FGuid& InBindingId, FMovieS
 
 		if (ULevelSequence* LevelSequence = Cast<ULevelSequence>(&InSequence))
 		{
-			// Passing through the streamed level asset path ensures that bindings within instance sub levels resolve correctly
-			LevelSequence->LocateBoundObjects(InBindingId, ResolutionContext, StreamedLevelAssetPath, OutObjects);
+			FLevelSequenceBindingReference::FResolveBindingParams Params;
+			Params.StreamedLevelAssetPath = StreamedLevelAssetPath;
+			LevelSequence->LocateBoundObjects(InBindingId, ResolutionContext, Params, OutObjects);
 		}
 		else
 		{
diff --git a/Engine/Source/Runtime/LevelSequence/Public/LevelSequence.h b/Engine/Source/Runtime/LevelSequence/Public/LevelSequence.h
--- a/Engine/Source/Runtime/LevelSequence/Public/LevelSequence.h
+++ b/Engine/Source/Runtime/LevelSequence/Public/LevelSequence.h
@@ -86,8 +86,11 @@ public:
 #endif
 
 	virtual void PostDuplicate(bool bDuplicateForPIE) override;
+		
+	UE_DEPRECATED(5.3, "Use LocateBoundObjects with FLocateBoundObjectsParams parameter instead")
+	void LocateBoundObjects(const FGuid& ObjectId, UObject* Context, const FTopLevelAssetPath& StreamedLevelAssetPath, TArray<UObject*, TInlineAllocator<1>>& OutObjects) const	{}
 
-	void LocateBoundObjects(const FGuid& ObjectId, UObject* Context, const FTopLevelAssetPath& StreamedLevelAssetPath, TArray<UObject*, TInlineAllocator<1>>& OutObjects) const;
+	void LocateBoundObjects(const FGuid& ObjectId, UObject* Context, const FLevelSequenceBindingReference::FResolveBindingParams& InResolveBindingParams, TArray<UObject*, TInlineAllocator<1>>& OutObjects) const;
 #if WITH_EDITOR
 
 
diff --git a/Engine/Source/Runtime/LevelSequence/Public/LevelSequenceBindingReference.h b/Engine/Source/Runtime/LevelSequence/Public/LevelSequenceBindingReference.h
--- a/Engine/Source/Runtime/LevelSequence/Public/LevelSequenceBindingReference.h
+++ b/Engine/Source/Runtime/LevelSequence/Public/LevelSequenceBindingReference.h
@@ -32,23 +32,32 @@ struct FLevelSequenceBindingReference
 	 */
 	LEVELSEQUENCE_API FLevelSequenceBindingReference(UObject* InObject, UObject* InContext);
 
+	/**
+	 * Structure that stores additional params that are used for resolving bindings.
+	 */
+	struct FResolveBindingParams
+	{
+		FResolveBindingParams() {}
+
+		// The path to the streamed level asset that contains the level sequence actor playing back the sequence. 'None' for any non - instance - level setups.
+		FTopLevelAssetPath StreamedLevelAssetPath;
+	};
+
 	/**
 	 * Resolve this reference within the specified context
 	 *
-	 * @param InContext		The context to resolve the binding within. Either a UWorld, ULevel (when playing in an instanced level) or an AActor where this binding relates to an actor component
-	 * @param StreamedLevelAssetPath    The path to the streamed level asset that contains the level sequence actor playing back the sequence. A null path for any non-instance-level setups.
-	 * @return The object (usually an Actor or an ActorComponent).
+	 * @param	InContext	The context to resolve the binding within. Either a UWorld, ULevel (when playing in an instanced level) or an AActor where this binding relates to an actor component
+	 * @param	InResolveBindingParams   The struct containing additional resolving params.
+	 * @return	The object (usually an Actor or an ActorComponent).
 	 */
-	LEVELSEQUENCE_API UObject* Resolve(UObject* InContext, const FTopLevelAssetPath& StreamedLevelAssetPath) const;
-	
-	UE_DEPRECATED(5.1, "Asset path FNames are deprecated. Use FTopLevelAssetPath instead.")
-	UObject* Resolve(UObject* InContext, FName StreamedLevelAssetPath) const
+	LEVELSEQUENCE_API UObject* Resolve(UObject* InContext, const FResolveBindingParams& InResolveBindingParams) const;
+		
+	UE_DEPRECATED(5.3, "Use Resolve taking FLevelSequenceBindingReference::FResolveBindingParams param instead")
+	UObject* Resolve(UObject* InContext, const FTopLevelAssetPath& StreamedLevelAssetPath) const
 	{
-PRAGMA_DISABLE_DEPRECATION_WARNINGS
-		return Resolve(InContext, FTopLevelAssetPath(StreamedLevelAssetPath));
-PRAGMA_ENABLE_DEPRECATION_WARNINGS		
+		return nullptr;
 	}
-
+	
 	/**
 	 * Check whether this binding reference is equal to the specified object
 	 */
@@ -142,15 +151,19 @@ struct FLevelSequenceBindingReferences
 	 */
 	void AddBinding(const FGuid& ObjectId, UObject* InObject, UObject* InContext);
 
+	UE_DEPRECATED(5.3, "Use ResolveBinding taking FLevelSequenceBindingReference::FResolveBindingParams param instead.")
+	void ResolveBinding(const FGuid& ObjectId, UObject* InContext, const FTopLevelAssetPath& StreamedLevelAssetPath, TArray<UObject*, TInlineAllocator<1>>& OutObjects) const { }
+		
 	/**
 	 * Resolve a binding for the specified ID using a given context
 	 *
 	 * @param ObjectId					The ID to associate the object with
 	 * @param InContext					A context in which InObject resides
-	 * @oaram StreamedLevelAssetPath    The path to the streamed level asset that contains the level sequence actor playing back the sequence. 'None' for any non-instance-level setups.
+	 * @oaram InResolveBindingParams    Struct containing additional resolving params
 	 * @param OutObjects				Array to populate with resolved object bindings
 	 */
-	void ResolveBinding(const FGuid& ObjectId, UObject* InContext, const FTopLevelAssetPath& StreamedLevelAssetPath, TArray<UObject*, TInlineAllocator<1>>& OutObjects) const;
+	void ResolveBinding(const FGuid& ObjectId, UObject* InContext, const FLevelSequenceBindingReference::FResolveBindingParams& InResolveBindingParams, TArray<UObject*, TInlineAllocator<1>>& OutObjects) const;
+
 
 	/**
 	 * Resolve a binding for the specified ID using a given context
