LevelSequence: Support Level Sequences within World Partition embedded Level Instances

#rnx
#tests LevelSequence test group + multitude of WP Level Instance use cases (regular WP and instanced)

diff --git a/Engine/Source/Runtime/Engine/Private/WorldPartition/WorldPartitionLevelHelper.cpp b/Engine/Source/Runtime/Engine/Private/WorldPartition/WorldPartitionLevelHelper.cpp
--- a/Engine/Source/Runtime/Engine/Private/WorldPartition/WorldPartitionLevelHelper.cpp
+++ b/Engine/Source/Runtime/Engine/Private/WorldPartition/WorldPartitionLevelHelper.cpp
@@ -27,7 +27,7 @@
 
 #endif
 
-bool FWorldPartitionResolveData::ResolveObject(UWorld* InWorld, const FSoftObjectPath& InObjectPath, UObject*& OutObject)
+bool FWorldPartitionResolveData::ResolveObject(UWorld* InWorld, const FSoftObjectPath& InObjectPath, UObject*& OutObject) const
 {
 	OutObject = nullptr;
 	if (InWorld)
diff --git a/Engine/Source/Runtime/Engine/Public/WorldPartition/IWorldPartitionObjectResolver.h b/Engine/Source/Runtime/Engine/Public/WorldPartition/IWorldPartitionObjectResolver.h
--- a/Engine/Source/Runtime/Engine/Public/WorldPartition/IWorldPartitionObjectResolver.h
+++ b/Engine/Source/Runtime/Engine/Public/WorldPartition/IWorldPartitionObjectResolver.h
@@ -19,7 +19,7 @@ struct FWorldPartitionResolveData
 
 	bool IsValid() const { return !ContainerID.IsMainContainer(); }
 		
-	bool ResolveObject(UWorld* InWorld, const FSoftObjectPath& InObjectPath, UObject*& OutObject);
+	ENGINE_API bool ResolveObject(UWorld* InWorld, const FSoftObjectPath& InObjectPath, UObject*& OutObject) const;
 
 	UPROPERTY()
 	FActorContainerID ContainerID;
diff --git a/Engine/Source/Runtime/LevelSequence/Private/LevelSequenceBindingReference.cpp b/Engine/Source/Runtime/LevelSequence/Private/LevelSequenceBindingReference.cpp
--- a/Engine/Source/Runtime/LevelSequence/Private/LevelSequenceBindingReference.cpp
+++ b/Engine/Source/Runtime/LevelSequence/Private/LevelSequenceBindingReference.cpp
@@ -14,6 +14,7 @@
 #include "Animation/AnimInstance.h"
 #include "Components/SkeletalMeshComponent.h"
 #include "Engine/LevelStreamingDynamic.h"
+#include "WorldPartition/IWorldPartitionObjectResolver.h"
 
 #include UE_INLINE_GENERATED_CPP_BY_NAME(LevelSequenceBindingReference)
 
@@ -49,6 +50,9 @@ FLevelSequenceBindingReference::FLevelSequenceBindingReference(UObject* InObject
 
 UObject* FLevelSequenceBindingReference::Resolve(UObject* InContext, const FResolveBindingParams& InResolveBindingParams) const
 {
+	// tidy up todo: StreamedLevelAsset path and WorldPartitionResolveData code paths could probably share most their code
+	//				 InContext could be reverted back to always be the UWorld in both paths and instead StreamedLevelAsset code path could use the StreamingWorld to resolve
+
 	if (InContext && InContext->IsA<AActor>())
 	{
 		if (ExternalObjectPath.IsNull())
@@ -74,6 +78,10 @@ UObject* FLevelSequenceBindingReference::Resolve(UObject* InContext, const FReso
 		check(ContextOuter);
 		ContextOuter->ResolveSubobject(*ExternalObjectPath.GetSubPathString(), ResolvedObject, /*bLoadIfExists*/false);
 		return ResolvedObject;
+	} 
+	else if (UObject* ResolvedObject = nullptr; ExternalObjectPath.IsValid() && InResolveBindingParams.WorldPartitionResolveData && InResolveBindingParams.WorldPartitionResolveData->ResolveObject(InResolveBindingParams.StreamingWorld, ExternalObjectPath, ResolvedObject))
+	{
+		return ResolvedObject;
 	}
 	else
 	{
diff --git a/Engine/Source/Runtime/LevelSequence/Private/LevelSequencePlayer.cpp b/Engine/Source/Runtime/LevelSequence/Private/LevelSequencePlayer.cpp
--- a/Engine/Source/Runtime/LevelSequence/Private/LevelSequencePlayer.cpp
+++ b/Engine/Source/Runtime/LevelSequence/Private/LevelSequencePlayer.cpp
@@ -96,12 +96,21 @@ void ULevelSequencePlayer::Initialize(ULevelSequence* InLevelSequence, ULevel* I
 	ULevelStreaming* LevelStreaming = FLevelUtils::FindStreamingLevel(InLevel);
 	if (LevelStreaming)
 	{
-		// All ULevelStreaming objects live in the owning world but if we are streaming a World Partition persistent level it will be returned as the StreamingWorld for all it's 
-		// ULevelStreaming cells. This streaming world should be used to resolve bindings.
-		if (UWorld* StreamingWorld = LevelStreaming->GetStreamingWorld(); StreamingWorld && (StreamingWorld != World))
+		// If we are streaming the persistent level of a World Partition
+		if (UWorldPartition* WorldPartition = InLevel->GetWorldPartition())
 		{
-			Level = StreamingWorld->PersistentLevel;
-			LevelStreaming = FLevelUtils::FindStreamingLevel(Level.Get());
+			StreamingWorld = InLevel->GetTypedOuter<UWorld>();
+		}
+		else
+		{
+			// All ULevelStreaming objects live in the owning world but if we are streaming a World Partition persistent level it will be returned as the StreamingWorld for all it's 
+			// ULevelStreaming cells. This streaming world should be used to resolve bindings.
+			StreamingWorld = LevelStreaming->GetStreamingWorld();
+			if (StreamingWorld.IsValid() && (StreamingWorld != World))
+			{
+				Level = StreamingWorld->PersistentLevel;
+				LevelStreaming = FLevelUtils::FindStreamingLevel(Level.Get());
+			}
 		}
 	}
 		
@@ -136,6 +145,14 @@ void ULevelSequencePlayer::ResolveBoundObjects(const FGuid& InBindingId, FMovieS
 		{
 			FLevelSequenceBindingReference::FResolveBindingParams Params;
 			Params.StreamedLevelAssetPath = StreamedLevelAssetPath;
+			
+			if (ALevelSequenceActor* LevelSequenceActor = GetTypedOuter<ALevelSequenceActor>(); LevelSequenceActor && LevelSequenceActor->GetWorldPartitionResolveData().IsValid())
+			{
+				Params.WorldPartitionResolveData = &LevelSequenceActor->GetWorldPartitionResolveData();
+				check(StreamingWorld.IsValid());
+				Params.StreamingWorld = StreamingWorld.Get();
+			}
+			
 			LevelSequence->LocateBoundObjects(InBindingId, ResolutionContext, Params, OutObjects);
 		}
 		else
diff --git a/Engine/Source/Runtime/LevelSequence/Public/LevelSequenceActor.h b/Engine/Source/Runtime/LevelSequence/Public/LevelSequenceActor.h
--- a/Engine/Source/Runtime/LevelSequence/Public/LevelSequenceActor.h
+++ b/Engine/Source/Runtime/LevelSequence/Public/LevelSequenceActor.h
@@ -12,6 +12,7 @@
 #include "MovieSceneBindingOwnerInterface.h"
 #include "MovieSceneBindingOverrides.h"
 #include "LevelSequenceCameraSettings.h"
+#include "WorldPartition/IWorldPartitionObjectResolver.h"
 
 #if UE_ENABLE_INCLUDE_ORDER_DEPRECATED_IN_5_1
 	#include "LevelSequencePlayer.h"
@@ -70,6 +71,7 @@ class LEVELSEQUENCE_API ALevelSequenceActor
 	: public AActor
 	, public IMovieScenePlaybackClient
 	, public IMovieSceneBindingOwnerInterface
+	, public IWorldPartitionObjectResolver
 {
 public:
 
@@ -286,8 +288,17 @@ protected:
 	//~ End AActor interface
 
 public:
+	const FWorldPartitionResolveData& GetWorldPartitionResolveData() const
+	{
+		return WorldPartitionResolveData;
+	}
 
 #if WITH_EDITOR
+	virtual void SetWorldPartitionResolveData(const FWorldPartitionResolveData& InWorldPartitionResolveData)
+	{
+		WorldPartitionResolveData = InWorldPartitionResolveData;
+	}
+
 	virtual bool GetReferencedContentObjects(TArray<UObject*>& Objects) const override;
 #endif //WITH_EDITOR
 
@@ -311,6 +322,9 @@ private:
 
 	UPROPERTY()
 	bool bShowBurnin;
+
+	UPROPERTY()
+	FWorldPartitionResolveData WorldPartitionResolveData;
 };
 
 USTRUCT()
diff --git a/Engine/Source/Runtime/LevelSequence/Public/LevelSequenceBindingReference.h b/Engine/Source/Runtime/LevelSequence/Public/LevelSequenceBindingReference.h
--- a/Engine/Source/Runtime/LevelSequence/Public/LevelSequenceBindingReference.h
+++ b/Engine/Source/Runtime/LevelSequence/Public/LevelSequenceBindingReference.h
@@ -10,6 +10,9 @@
 #include "Misc/Paths.h"
 #include "LevelSequenceBindingReference.generated.h"
 
+class UWorld;
+struct FWorldPartitionResolveData;
+
 /**
  * An external reference to an level sequence object, resolvable through an arbitrary context.
  * 
@@ -37,10 +40,16 @@ struct FLevelSequenceBindingReference
 	 */
 	struct FResolveBindingParams
 	{
-		FResolveBindingParams() {}
+		FResolveBindingParams() : WorldPartitionResolveData(nullptr), StreamingWorld(nullptr) {}
 
 		// The path to the streamed level asset that contains the level sequence actor playing back the sequence. 'None' for any non - instance - level setups.
 		FTopLevelAssetPath StreamedLevelAssetPath;
+
+		// World Partition Resolve Data
+		const FWorldPartitionResolveData* WorldPartitionResolveData;
+
+		// World Partition Streaming World
+		UWorld* StreamingWorld;
 	};
 
 	/**
diff --git a/Engine/Source/Runtime/LevelSequence/Public/LevelSequencePlayer.h b/Engine/Source/Runtime/LevelSequence/Public/LevelSequencePlayer.h
--- a/Engine/Source/Runtime/LevelSequence/Public/LevelSequencePlayer.h
+++ b/Engine/Source/Runtime/LevelSequence/Public/LevelSequencePlayer.h
@@ -172,6 +172,9 @@ private:
 	/** The full asset path (/Game/Folder/MapName.MapName) of the streaming level this player resides within. Bindings to actors with the same FSoftObjectPath::GetAssetPath are resolved within the cached level, rather than globally.. */
 	FTopLevelAssetPath StreamedLevelAssetPath;
 
+	/** The world to use for resolving FWorldPartitionResolveData */
+	TWeakObjectPtr<UWorld> StreamingWorld;
+
 	/** The camera settings to use when playing the sequence */
 	FLevelSequenceCameraSettings CameraSettings;
 
