Sequencer: Speculative fix for null movie scene crash on PostLoad

#rnx

diff --git a/Engine/Source/Runtime/LevelSequence/Private/LevelSequence.cpp b/Engine/Source/Runtime/LevelSequence/Private/LevelSequence.cpp
--- a/Engine/Source/Runtime/LevelSequence/Private/LevelSequence.cpp
+++ b/Engine/Source/Runtime/LevelSequence/Private/LevelSequence.cpp
@@ -310,29 +310,32 @@ void ULevelSequence::PostLoad()
 		}
 	}
 
-	TSet<FGuid> InvalidSpawnables;
-
-	for (int32 Index = 0; Index < MovieScene->GetSpawnableCount(); ++Index)
+	if (MovieScene)
 	{
-		FMovieSceneSpawnable& Spawnable = MovieScene->GetSpawnable(Index);
-		if (!Spawnable.GetObjectTemplate())
+		TSet<FGuid> InvalidSpawnables;
+
+		for (int32 Index = 0; Index < MovieScene->GetSpawnableCount(); ++Index)
 		{
-			if (Spawnable.GeneratedClass_DEPRECATED && Spawnable.GeneratedClass_DEPRECATED->ClassGeneratedBy)
+			FMovieSceneSpawnable& Spawnable = MovieScene->GetSpawnable(Index);
+			if (!Spawnable.GetObjectTemplate())
 			{
-				const FName TemplateName = MakeUniqueObjectName(MovieScene, UObject::StaticClass(), Spawnable.GeneratedClass_DEPRECATED->ClassGeneratedBy->GetFName());
-
-				UObject* NewTemplate = NewObject<UObject>(MovieScene, Spawnable.GeneratedClass_DEPRECATED->GetSuperClass(), TemplateName);
-				if (NewTemplate)
+				if (Spawnable.GeneratedClass_DEPRECATED && Spawnable.GeneratedClass_DEPRECATED->ClassGeneratedBy)
 				{
-					Spawnable.CopyObjectTemplate(*NewTemplate, *this);
+					const FName TemplateName = MakeUniqueObjectName(MovieScene, UObject::StaticClass(), Spawnable.GeneratedClass_DEPRECATED->ClassGeneratedBy->GetFName());
+
+					UObject* NewTemplate = NewObject<UObject>(MovieScene, Spawnable.GeneratedClass_DEPRECATED->GetSuperClass(), TemplateName);
+					if (NewTemplate)
+					{
+						Spawnable.CopyObjectTemplate(*NewTemplate, *this);
+					}
 				}
 			}
-		}
 
-		if (!Spawnable.GetObjectTemplate())
-		{
-			InvalidSpawnables.Add(Spawnable.GetGuid());
-			UE_LOG(LogLevelSequence, Warning, TEXT("Spawnable '%s' with ID '%s' does not have a valid object template"), *Spawnable.GetName(), *Spawnable.GetGuid().ToString());
+			if (!Spawnable.GetObjectTemplate())
+			{
+				InvalidSpawnables.Add(Spawnable.GetGuid());
+				UE_LOG(LogLevelSequence, Warning, TEXT("Spawnable '%s' with ID '%s' does not have a valid object template"), *Spawnable.GetName(), *Spawnable.GetGuid().ToString());
+			}
 		}
 	}
 
