Sequencer: Linked Anim Sequences, make them editor only

diff --git a/Engine/Source/Runtime/LevelSequence/Public/AnimSequenceLevelSequenceLink.h b/Engine/Source/Runtime/LevelSequence/Public/AnimSequenceLevelSequenceLink.h
--- a/Engine/Source/Runtime/LevelSequence/Public/AnimSequenceLevelSequenceLink.h
+++ b/Engine/Source/Runtime/LevelSequence/Public/AnimSequenceLevelSequenceLink.h
@@ -13,6 +13,8 @@ class LEVELSEQUENCE_API UAnimSequenceLevelSequenceLink : public UAssetUserData
 	GENERATED_UCLASS_BODY()
 
 public:
+	virtual bool IsEditorOnly() const override { return true; }
+
 	UPROPERTY(BlueprintReadWrite, Category = Property)
 	FGuid SkelTrackGuid;
 
diff --git a/Engine/Source/Runtime/LevelSequence/Public/LevelSequenceAnimSequenceLink.h b/Engine/Source/Runtime/LevelSequence/Public/LevelSequenceAnimSequenceLink.h
--- a/Engine/Source/Runtime/LevelSequence/Public/LevelSequenceAnimSequenceLink.h
+++ b/Engine/Source/Runtime/LevelSequence/Public/LevelSequenceAnimSequenceLink.h
@@ -58,6 +58,8 @@ class LEVELSEQUENCE_API ULevelSequenceAnimSequenceLink : public UAssetUserData
 	GENERATED_UCLASS_BODY()
 
 public:
+	virtual bool IsEditorOnly() const override { return true; }
+
 	UPROPERTY(BlueprintReadWrite, Category = Links)
 	TArray< FLevelSequenceAnimSequenceLinkItem> AnimSequenceLinks;
 };
\ No newline at end of file
