Sequencer: carry over the previous view target from a previously playing level sequence, if any

diff --git a/Engine/Source/Runtime/LevelSequence/Private/LevelSequencePlayer.cpp b/Engine/Source/Runtime/LevelSequence/Private/LevelSequencePlayer.cpp
--- a/Engine/Source/Runtime/LevelSequence/Private/LevelSequencePlayer.cpp
+++ b/Engine/Source/Runtime/LevelSequence/Private/LevelSequencePlayer.cpp
@@ -362,9 +362,40 @@ void ULevelSequencePlayer::UpdateCameraCut(UObject* CameraObject, const EMovieSc
 	}
 
 	// Save the last view target/aspect ratio constraint/etc. so that it can all be restored when the camera object is null.
-	if (!LastViewTarget.IsValid() || (ViewTarget != CameraObject && ViewTarget != LastCameraObject))
+	if (!LastViewTarget.IsValid())
 	{
-		LastViewTarget = ViewTarget;
+		// If we had another level sequence playing before us, we don't want to cache its last camera. Let's see if
+		// we can instead ask it what its LastViewTarget was, so that we get what the actual previous view target was
+		// before any sequence started to play.
+		TArray<IMovieScenePlayer*> ActivePlayers;
+		const bool bOnlyEvaluatingPlayers = true;
+		IMovieScenePlayer::Get(ActivePlayers, bOnlyEvaluatingPlayers);
+		for (IMovieScenePlayer* ActivePlayer : ActivePlayers)
+		{
+			UObject* PlayerObject = ActivePlayer->AsUObject();
+			if (!PlayerObject)
+			{
+				continue;
+			}
+
+
+			ULevelSequencePlayer* PreviousLevelSequencePlayer = Cast<ULevelSequencePlayer>(PlayerObject);
+			if (!PreviousLevelSequencePlayer || PreviousLevelSequencePlayer == this)
+			{
+				continue;
+			}
+
+			if (PreviousLevelSequencePlayer->LastViewTarget.IsValid())
+			{
+				LastViewTarget = PreviousLevelSequencePlayer->LastViewTarget;
+				break;
+			}
+		}
+
+		if (!LastViewTarget.IsValid() && (ViewTarget != CameraObject && ViewTarget != LastCameraObject))
+		{
+			LastViewTarget = ViewTarget;
+		}
 	}
 	if (!LastLocalPlayer.IsValid() || (LocalPlayer != LastLocalPlayer))
 	{
diff --git a/Engine/Source/Runtime/MovieScene/Private/MovieScenePlayer.cpp b/Engine/Source/Runtime/MovieScene/Private/MovieScenePlayer.cpp
--- a/Engine/Source/Runtime/MovieScene/Private/MovieScenePlayer.cpp
+++ b/Engine/Source/Runtime/MovieScene/Private/MovieScenePlayer.cpp
@@ -5,6 +5,7 @@
 #include "Evaluation/MovieSceneEvaluationTemplateInstance.h"
 #include "EntitySystem/MovieSceneSequenceInstance.h"
 #include "Misc/ScopeRWLock.h"
+#include "MovieSceneFwd.h"
 #include "MovieSceneSequence.h"
 #include "MovieSceneSequenceID.h"
 
@@ -46,6 +47,21 @@ IMovieScenePlayer* IMovieScenePlayer::Get(uint16 InUniqueIndex)
 	return UE::MovieScene::GGlobalPlayerRegistry[InUniqueIndex];
 }
 
+void IMovieScenePlayer::Get(TArray<IMovieScenePlayer*>& OutPlayers, bool bOnlyUnstoppedPlayers)
+{
+	FReadScopeLock ScopeLock(UE::MovieScene::GGlobalPlayerRegistryLock);
+	for (auto It = UE::MovieScene::GGlobalPlayerRegistry.CreateIterator(); It; ++It)
+	{
+		if (IMovieScenePlayer* Player = *It)
+		{
+			if (!bOnlyUnstoppedPlayers || Player->GetPlaybackStatus() != EMovieScenePlayerStatus::Stopped)
+			{
+				OutPlayers.Add(*It);
+			}
+		}
+	}
+}
+
 void IMovieScenePlayer::SetIsEvaluatingFlag(uint16 InUniqueIndex, bool bIsUpdating)
 {
 	check(UE::MovieScene::GGlobalPlayerUpdateFlags.IsValidIndex(InUniqueIndex));
@@ -78,4 +94,4 @@ void IMovieScenePlayer::InvalidateCachedData()
 	{
 		RootInstance->InvalidateCachedData(Template.GetEntitySystemLinker());
 	}
-}
\ No newline at end of file
+}
diff --git a/Engine/Source/Runtime/MovieScene/Public/IMovieScenePlayer.h b/Engine/Source/Runtime/MovieScene/Public/IMovieScenePlayer.h
--- a/Engine/Source/Runtime/MovieScene/Public/IMovieScenePlayer.h
+++ b/Engine/Source/Runtime/MovieScene/Public/IMovieScenePlayer.h
@@ -346,6 +346,8 @@ public:
 
 	MOVIESCENE_API static IMovieScenePlayer* Get(uint16 InUniqueIndex);
 
+	MOVIESCENE_API static void Get(TArray<IMovieScenePlayer*>& OutPlayers, bool bOnlyUnstoppedPlayers);
+
 	MOVIESCENE_API static void SetIsEvaluatingFlag(uint16 InUniqueIndex, bool bIsUpdating);
 
 	MOVIESCENE_API bool IsEvaluating() const;
