Sequencer: correctly reset the last known serial number when handling replays.

This was causing all replay messages to be ignored for being too old.

diff --git a/Engine/Source/Runtime/LevelSequence/Private/LevelSequencePlayer.cpp b/Engine/Source/Runtime/LevelSequence/Private/LevelSequencePlayer.cpp
--- a/Engine/Source/Runtime/LevelSequence/Private/LevelSequencePlayer.cpp
+++ b/Engine/Source/Runtime/LevelSequence/Private/LevelSequencePlayer.cpp
@@ -662,5 +662,6 @@ void ULevelSequencePlayer::RewindForReplay()
 	NetSyncProps.LastKnownPosition = FFrameTime(0);
 	NetSyncProps.LastKnownStatus = EMovieScenePlayerStatus::Stopped;
 	NetSyncProps.LastKnownNumLoops = 0;
+	NetSyncProps.LastKnownSerialNumber = 0;
 }
 
