Sequencer: fix camera cut flag not being set when spawnable cameras are the last in control

This was because calling GetViewTarget does some internal checks. If a spawnable camera has just been unspawned (i.e. the actor is PendingKill) then this trips those checks and the view target is reset to the pawn before we can even do anything. Because we are restoring the view target to the pawn, this is registered as "not changing anything" and we skip everything in UpdateCameraCut, including setting the camera cut flag.

We therefore don't use GetViewTarget, and grab the view target directly from the public member to manage things by hand.

diff --git a/Engine/Source/Runtime/LevelSequence/Private/LevelSequencePlayer.cpp b/Engine/Source/Runtime/LevelSequence/Private/LevelSequencePlayer.cpp
--- a/Engine/Source/Runtime/LevelSequence/Private/LevelSequencePlayer.cpp
+++ b/Engine/Source/Runtime/LevelSequence/Private/LevelSequencePlayer.cpp
@@ -281,10 +281,12 @@ void ULevelSequencePlayer::UpdateCameraCut(UObject* CameraObject, const EMovieSc
 		return;
 	}
 
-	// skip missing player controller
+	// If the player controller is missing, there is no camera manager for us to manage the view target
+	// so we just bail out.
 	APlayerController* PC = World->GetGameInstance()->GetFirstLocalPlayerController();
+	APlayerCameraManager* CameraManager = (PC != nullptr) ? PC->PlayerCameraManager : nullptr;
 
-	if (PC == nullptr)
+	if (CameraManager == nullptr)
 	{
 		// If we are releasing view target control but can't actually do it because game state has been
 		// torn down, see if we can at least restore aspect ratio axis constraint. This fixes a bug
@@ -300,14 +302,24 @@ void ULevelSequencePlayer::UpdateCameraCut(UObject* CameraObject, const EMovieSc
 		return;
 	}
 
-	// skip same view target
-	AActor* ViewTarget = PC->GetViewTarget();
-
+	// If we don't want to update camera cuts, let's also bail out.
 	if (!CanUpdateCameraCut())
 	{
 		return;
 	}
 
+	// Let's get the current view target directly from the player camera manager first. This is because
+	// we don't want to go through GetViewTarget, which checks if the current view target is valid, and
+	// re-assigns it to the player controller if it's not. We don't want this to happen, especially since
+	// a spawnable camera might have just been unspawned, causing UpdateCameraCut to be called, and we
+	// need to handle this properly.
+	AActor* ViewTarget = CameraManager->PendingViewTarget.Target;
+	if (!ViewTarget)
+	{
+		ViewTarget = CameraManager->ViewTarget.Target;
+	}
+
+	// If the view target isn't really changing, we don't have much to do.
 	if (CameraObject == ViewTarget)
 	{
 		if (CameraCutParams.bJumpCut)
