Sequencer: fix static analysis warnings

diff --git a/Engine/Source/Runtime/LevelSequence/Private/LevelSequencePlayer.cpp b/Engine/Source/Runtime/LevelSequence/Private/LevelSequencePlayer.cpp
--- a/Engine/Source/Runtime/LevelSequence/Private/LevelSequencePlayer.cpp
+++ b/Engine/Source/Runtime/LevelSequence/Private/LevelSequencePlayer.cpp
@@ -324,9 +324,9 @@ void ULevelSequencePlayer::UpdateCameraCut(UObject* CameraObject, const EMovieSc
 	{
 		if (CameraCutParams.bJumpCut)
 		{
-			if (PC->PlayerCameraManager)
+			if (CameraManager)
 			{
-				PC->PlayerCameraManager->SetGameCameraCutThisFrame();
+				CameraManager->SetGameCameraCutThisFrame();
 			}
 
 			if (CameraComponent)
@@ -353,7 +353,7 @@ void ULevelSequencePlayer::UpdateCameraCut(UObject* CameraObject, const EMovieSc
 
 	// override the player controller's view target
 	AActor* CameraActor = Cast<AActor>(CameraObject);
-	ULocalPlayer* LocalPlayer = PC->GetLocalPlayer();
+	ULocalPlayer* LocalPlayer = (PC != nullptr) ? PC->GetLocalPlayer() : nullptr;
 
 	// if the camera object is null, use the last view target so that it is restored to the state before the sequence takes control
 	bool bRestoreAspectRatioConstraint = false;
@@ -444,10 +444,9 @@ void ULevelSequencePlayer::UpdateCameraCut(UObject* CameraObject, const EMovieSc
 		// Calling SetViewTarget on a camera that we are currently transitioning to will 
 		// result in that transition being aborted, and the view target being set immediately.
 		// We want to avoid that, so let's leave the transition running if it's the case.
-		if (PC->PlayerCameraManager != nullptr)
+		if (CameraManager != nullptr)
 		{
-			const AActor* CurViewTarget = PC->PlayerCameraManager->ViewTarget.Target;
-			const AActor* PendingViewTarget = PC->PlayerCameraManager->PendingViewTarget.Target;
+			const AActor* PendingViewTarget = CameraManager->PendingViewTarget.Target;
 			if (CameraActor != nullptr && PendingViewTarget == CameraActor)
 			{
 				UE_LOG(LogLevelSequence, Log, TEXT("Camera transition aborted, we are already blending towards the intended camera"));
@@ -460,7 +459,7 @@ void ULevelSequencePlayer::UpdateCameraCut(UObject* CameraObject, const EMovieSc
 		UE_LOG(LogLevelSequence, Log, TEXT("Starting new camera cut: '%s'"),
 			(CameraObject ? *CameraObject->GetName() : TEXT("None")));
 	}
-	if (bDoSetViewTarget)
+	if (bDoSetViewTarget && ensureMsgf(PC, TEXT("Can't set view target when there is no player controller!")))
 	{
 		PC->SetViewTarget(CameraActor, TransitionParams);
 	}
@@ -490,13 +489,13 @@ void ULevelSequencePlayer::UpdateCameraCut(UObject* CameraObject, const EMovieSc
 		CameraComponent->NotifyCameraCut();
 	}
 
-	if (PC->PlayerCameraManager)
+	if (CameraManager)
 	{
-		PC->PlayerCameraManager->bClientSimulatingViewTarget = (CameraActor != nullptr);
+		CameraManager->bClientSimulatingViewTarget = (CameraActor != nullptr);
 
 		if (bIsStraightCut)
 		{
-			PC->PlayerCameraManager->SetGameCameraCutThisFrame();
+			CameraManager->SetGameCameraCutThisFrame();
 		}
 	}
 
