LevelSequence: Fix LevelInstance actor resolving of AlwaysLoaded LevelInstance actors and Disabled Streaming World Partitions

#rnx

diff --git a/Engine/Source/Runtime/LevelSequence/Private/LevelSequencePlayer.cpp b/Engine/Source/Runtime/LevelSequence/Private/LevelSequencePlayer.cpp
--- a/Engine/Source/Runtime/LevelSequence/Private/LevelSequencePlayer.cpp
+++ b/Engine/Source/Runtime/LevelSequence/Private/LevelSequencePlayer.cpp
@@ -91,7 +91,8 @@ void ULevelSequencePlayer::Initialize(ULevelSequence* InLevelSequence, ULevel* I
 	World = InLevel->OwningWorld;
 	Level = InLevel;
 	CameraSettings = InCameraSettings;
-
+	// Default to owning world (to resolve AlwaysLoaded actors not part of a Streaming Level and Disabled Streaming World Partitions)
+	StreamingWorld = World;
 	// Construct the path to the level asset that the streamed level relates to
 	ULevelStreaming* LevelStreaming = FLevelUtils::FindStreamingLevel(InLevel);
 	if (LevelStreaming)
