Sequencer: Make sure setting the last view target is done even in early out. This fixes an issue where if sequence 1 plays and cuts to camera 1, and the sequence 2 plays and also cuts to camera 2, it restores properly when the sequences finish playing.

diff --git a/Engine/Source/Runtime/LevelSequence/Private/LevelSequencePlayer.cpp b/Engine/Source/Runtime/LevelSequence/Private/LevelSequencePlayer.cpp
--- a/Engine/Source/Runtime/LevelSequence/Private/LevelSequencePlayer.cpp
+++ b/Engine/Source/Runtime/LevelSequence/Private/LevelSequencePlayer.cpp
@@ -267,6 +267,45 @@ TTuple<EViewTargetBlendFunction, float> BuiltInEasingTypeToBlendFunction(EMovieS
 	return Return(EViewTargetBlendFunction::VTBlend_Linear, 1.f);
 }
 
+void ULevelSequencePlayer::ValidateLastViewTarget(UObject* CameraObject, AActor* ViewTarget)
+{
+	if (!LastViewTarget.IsValid())
+	{
+		// If we had another level sequence playing before us, we don't want to cache its last camera. Let's see if
+		// we can instead ask it what its LastViewTarget was, so that we get what the actual previous view target was
+		// before any sequence started to play.
+		TArray<IMovieScenePlayer*> ActivePlayers;
+		const bool bOnlyEvaluatingPlayers = true;
+		IMovieScenePlayer::Get(ActivePlayers, bOnlyEvaluatingPlayers);
+		for (IMovieScenePlayer* ActivePlayer : ActivePlayers)
+		{
+			UObject* PlayerObject = ActivePlayer->AsUObject();
+			if (!PlayerObject)
+			{
+				continue;
+			}
+
+
+			ULevelSequencePlayer* PreviousLevelSequencePlayer = Cast<ULevelSequencePlayer>(PlayerObject);
+			if (!PreviousLevelSequencePlayer || PreviousLevelSequencePlayer == this)
+			{
+				continue;
+			}
+
+			if (PreviousLevelSequencePlayer->LastViewTarget.IsValid())
+			{
+				LastViewTarget = PreviousLevelSequencePlayer->LastViewTarget;
+				break;
+			}
+		}
+	}
+
+	if (!LastViewTarget.IsValid() || (ViewTarget != nullptr && ViewTarget != CameraObject && ViewTarget != LastCameraObject))
+	{
+		LastViewTarget = ViewTarget;
+	}
+}
+
 void ULevelSequencePlayer::UpdateCameraCut(UObject* CameraObject, const EMovieSceneCameraCutParams& CameraCutParams)
 {
 	UCameraComponent* CameraComponent = MovieSceneHelpers::CameraComponentFromRuntimeObject(CameraObject);
@@ -340,6 +379,12 @@ void ULevelSequencePlayer::UpdateCameraCut(UObject* CameraObject, const EMovieSc
 				MotionVectorSim->SimulateAllTransforms();
 			}
 		}
+
+		// If we just started a new sequence that happens to give control to the same camera actor as
+		// the previous sequence, we end up here for an early out, but we also miss on setting out LastViewTarget,
+		// so let's check we have a valid one.
+		ValidateLastViewTarget(CameraObject, ViewTarget);
+
 		return;
 	}
 
@@ -361,9 +406,10 @@ void ULevelSequencePlayer::UpdateCameraCut(UObject* CameraObject, const EMovieSc
 	if (CameraActor == nullptr)
 	{
 		CameraActor = LastViewTarget.Get();
-		bRestoreAspectRatioConstraint = true;
+		CameraObject = CameraActor;
 
 		// Skip if the last view target is the same as the current view target so that there's no additional camera cut
+		// In this case, restore the aspect ratio axis constraint right away before exiting.
 		if (CameraActor == ViewTarget)
 		{
 			if (LocalPlayer && LastAspectRatioAxisConstraint.IsSet())
@@ -372,44 +418,13 @@ void ULevelSequencePlayer::UpdateCameraCut(UObject* CameraObject, const EMovieSc
 			}
 			return;
 		}
+		// Otherwise, remember to restore the aspect ratio axis constraint later in this function.
+		bRestoreAspectRatioConstraint = true;
 	}
 
 	// Save the last view target/aspect ratio constraint/etc. so that it can all be restored when the camera object is null.
-	if (!LastViewTarget.IsValid())
-	{
-		// If we had another level sequence playing before us, we don't want to cache its last camera. Let's see if
-		// we can instead ask it what its LastViewTarget was, so that we get what the actual previous view target was
-		// before any sequence started to play.
-		TArray<IMovieScenePlayer*> ActivePlayers;
-		const bool bOnlyEvaluatingPlayers = true;
-		IMovieScenePlayer::Get(ActivePlayers, bOnlyEvaluatingPlayers);
-		for (IMovieScenePlayer* ActivePlayer : ActivePlayers)
-		{
-			UObject* PlayerObject = ActivePlayer->AsUObject();
-			if (!PlayerObject)
-			{
-				continue;
-			}
-
-
-			ULevelSequencePlayer* PreviousLevelSequencePlayer = Cast<ULevelSequencePlayer>(PlayerObject);
-			if (!PreviousLevelSequencePlayer || PreviousLevelSequencePlayer == this)
-			{
-				continue;
-			}
-
-			if (PreviousLevelSequencePlayer->LastViewTarget.IsValid())
-			{
-				LastViewTarget = PreviousLevelSequencePlayer->LastViewTarget;
-				break;
-			}
-		}
+	ValidateLastViewTarget(CameraObject, ViewTarget);
 
-		if (!LastViewTarget.IsValid() && (ViewTarget != CameraObject && ViewTarget != LastCameraObject))
-		{
-			LastViewTarget = ViewTarget;
-		}
-	}
 	if (!LastLocalPlayer.IsValid() || (LocalPlayer != LastLocalPlayer))
 	{
 		LastLocalPlayer = LocalPlayer;
diff --git a/Engine/Source/Runtime/LevelSequence/Public/LevelSequencePlayer.h b/Engine/Source/Runtime/LevelSequence/Public/LevelSequencePlayer.h
--- a/Engine/Source/Runtime/LevelSequence/Public/LevelSequencePlayer.h
+++ b/Engine/Source/Runtime/LevelSequence/Public/LevelSequencePlayer.h
@@ -151,6 +151,9 @@ private:
 
 	LEVELSEQUENCE_API void EnableCinematicMode(bool bEnable);
 
+	// Save the last view target, so that it can all be restored when the camera object is null.
+	void ValidateLastViewTarget(UObject* CameraObject, AActor* ViewTarget);
+
 private:
 
 	/** The world this player will spawn actors in, if needed */
