Sequencer: don't add a playback capability if it already exists

This can happen if a level sequence is re-initialized with a sequence that it was already set to.

diff --git a/Engine/Source/Runtime/LevelSequence/Private/LevelSequencePlayer.cpp b/Engine/Source/Runtime/LevelSequence/Private/LevelSequencePlayer.cpp
--- a/Engine/Source/Runtime/LevelSequence/Private/LevelSequencePlayer.cpp
+++ b/Engine/Source/Runtime/LevelSequence/Private/LevelSequencePlayer.cpp
@@ -140,7 +140,10 @@ void ULevelSequencePlayer::Initialize(ULevelSequence* InLevelSequence, ULevel* I
 	UMovieSceneEntitySystemLinker* Linker = GetEvaluationTemplate().GetEntitySystemLinker();
 	FInstanceRegistry* InstanceRegistry = Linker->GetInstanceRegistry();
 	TSharedRef<FSharedPlaybackState> SharedPlaybackState = InstanceRegistry->GetInstance(RootInstanceHandle).GetSharedPlaybackState();
-	SharedPlaybackState->AddCapabilityRaw<FCameraCutPlaybackCapability>((FCameraCutPlaybackCapability*)this);
+	if (!SharedPlaybackState->HasCapability<FCameraCutPlaybackCapability>())
+	{
+		SharedPlaybackState->AddCapabilityRaw<FCameraCutPlaybackCapability>((FCameraCutPlaybackCapability*)this);
+	}
 }
 
 void ULevelSequencePlayer::ResolveBoundObjects(const FGuid& InBindingId, FMovieSceneSequenceID SequenceID, UMovieSceneSequence& InSequence, UObject* ResolutionContext, TArray<UObject*, TInlineAllocator<1>>& OutObjects) const
