Sequencer: Modify when adding a section for CreateSpawnable

This fixes the UI not updating when calling AddSpawnableFromClass from an Editor Utility Widget

#rnx

diff --git a/Engine/Source/Runtime/LevelSequence/Private/LevelSequence.cpp b/Engine/Source/Runtime/LevelSequence/Private/LevelSequence.cpp
--- a/Engine/Source/Runtime/LevelSequence/Private/LevelSequence.cpp
+++ b/Engine/Source/Runtime/LevelSequence/Private/LevelSequence.cpp
@@ -690,6 +690,8 @@ FGuid ULevelSequence::CreateSpawnable(UObject* ObjectToSpawn)
 			UMovieSceneSpawnTrack* NewSpawnTrack = MovieScene->AddTrack<UMovieSceneSpawnTrack>(NewGuid);
 			if (NewSpawnTrack)
 			{
+				NewSpawnTrack->Modify();
+
 				NewSpawnTrack->AddSection(*NewSpawnTrack->CreateNewSection());
 			}
 			return NewGuid;
