Fix LevelSequenceActor replication in Iris by converting them to use the replicated subobject list.
MoveSceneSequencePlayer now creates Iris replication fragments correctly.

#rnx

diff --git a/Engine/Source/Runtime/LevelSequence/Private/LevelSequenceActor.cpp b/Engine/Source/Runtime/LevelSequence/Private/LevelSequenceActor.cpp
--- a/Engine/Source/Runtime/LevelSequence/Private/LevelSequenceActor.cpp
+++ b/Engine/Source/Runtime/LevelSequence/Private/LevelSequenceActor.cpp
@@ -63,8 +63,6 @@ ALevelSequenceActor::ALevelSequenceActor(const FObjectInitializer& Init)
 	bIsSpatiallyLoaded = false;
 #endif //WITH_EDITORONLY_DATA
 
-	bReplicateUsingRegisteredSubObjectList = false;
-
 	BindingOverrides = Init.CreateDefaultSubobject<UMovieSceneBindingOverrides>(this, "BindingOverrides");
 	BurnInOptions = Init.CreateDefaultSubobject<ULevelSequenceBurnInOptions>(this, "BurnInOptions");
 	DefaultInstanceData = Init.CreateDefaultSubobject<UDefaultLevelSequenceInstanceData>(this, "InstanceData");
@@ -139,15 +137,6 @@ void ALevelSequenceActor::SetReplicatePlayback(bool bInReplicatePlayback)
 	SetReplicates(bReplicatePlayback);
 }
 
-bool ALevelSequenceActor::ReplicateSubobjects(UActorChannel* Channel, FOutBunch* Bunch, FReplicationFlags* RepFlags)
-{
-	bool bWroteSomething = Super::ReplicateSubobjects(Channel, Bunch, RepFlags);
-
-	bWroteSomething |= Channel->ReplicateSubobject(SequencePlayer, *Bunch, *RepFlags);
-
-	return bWroteSomething;
-}
-
 void ALevelSequenceActor::GetLifetimeReplicatedProps(TArray<FLifetimeProperty>& OutLifetimeProps) const
 {
 	Super::GetLifetimeReplicatedProps(OutLifetimeProps);
@@ -176,6 +165,11 @@ void ALevelSequenceActor::BeginPlay()
 {
 	Super::BeginPlay();
 
+	if (SequencePlayer)
+	{
+		AddReplicatedSubObject(SequencePlayer);
+	}
+
 	if (PlaybackSettings.bAutoPlay)
 	{
 		SequencePlayer->Play();
@@ -186,6 +180,8 @@ void ALevelSequenceActor::EndPlay(const EEndPlayReason::Type EndPlayReason)
 {
 	if (SequencePlayer)
 	{
+		RemoveReplicatedSubObject(SequencePlayer);
+
 		// Stop may modify a lot of actor state so it needs to be called
 		// during EndPlay (when Actors + World are still valid) instead
 		// of waiting for the UObject to be destroyed by GC.
diff --git a/Engine/Source/Runtime/LevelSequence/Public/LevelSequenceActor.h b/Engine/Source/Runtime/LevelSequence/Public/LevelSequenceActor.h
--- a/Engine/Source/Runtime/LevelSequence/Public/LevelSequenceActor.h
+++ b/Engine/Source/Runtime/LevelSequence/Public/LevelSequenceActor.h
@@ -261,7 +261,6 @@ protected:
 	//~ End IMovieScenePlaybackClient interface
 
 	//~ Begin UObject interface
-	LEVELSEQUENCE_API virtual bool ReplicateSubobjects(UActorChannel* Channel, FOutBunch* Bunch, FReplicationFlags *RepFlags) override;
 	LEVELSEQUENCE_API virtual void PostInitProperties() override;
 	LEVELSEQUENCE_API virtual void PostLoad() override;
 public:
diff --git a/Engine/Source/Runtime/MovieScene/MovieScene.Build.cs b/Engine/Source/Runtime/MovieScene/MovieScene.Build.cs
--- a/Engine/Source/Runtime/MovieScene/MovieScene.Build.cs
+++ b/Engine/Source/Runtime/MovieScene/MovieScene.Build.cs
@@ -21,6 +21,6 @@ public class MovieScene : ModuleRules
 				"TimeManagement"
 			}
 		);
-
-    }
+		SetupIrisSupport(Target);
+	}
 }
diff --git a/Engine/Source/Runtime/MovieScene/Private/MovieSceneSequencePlayer.cpp b/Engine/Source/Runtime/MovieScene/Private/MovieSceneSequencePlayer.cpp
--- a/Engine/Source/Runtime/MovieScene/Private/MovieSceneSequencePlayer.cpp
+++ b/Engine/Source/Runtime/MovieScene/Private/MovieSceneSequencePlayer.cpp
@@ -22,6 +22,11 @@
 #include "GameFramework/PlayerState.h"
 #include "Algo/BinarySearch.h"
 
+#if UE_WITH_IRIS
+#include "Iris/ReplicationSystem/ReplicationFragmentUtil.h"
+#include "Net/Iris/ReplicationSystem/ReplicationSystemUtil.h"
+#endif // UE_WITH_IRIS
+
 
 #include UE_INLINE_GENERATED_CPP_BY_NAME(MovieSceneSequencePlayer)
 
@@ -1969,6 +1974,13 @@ bool UMovieSceneSequencePlayer::CallRemoteFunction(UFunction* Function, void* Pa
 	return false;
 }
 
+#if UE_WITH_IRIS
+void UMovieSceneSequencePlayer::RegisterReplicationFragments(UE::Net::FFragmentRegistrationContext& Context, UE::Net::EFragmentRegistrationFlags RegistrationFlags)
+{
+	UE::Net::FReplicationFragmentUtil::CreateAndRegisterFragmentsForObject(this, Context, RegistrationFlags);
+}
+#endif
+
 bool UMovieSceneSequencePlayer::NeedsQueueLatentAction() const
 {
 	return IsEvaluating();
diff --git a/Engine/Source/Runtime/MovieScene/Public/MovieSceneSequencePlayer.h b/Engine/Source/Runtime/MovieScene/Public/MovieSceneSequencePlayer.h
--- a/Engine/Source/Runtime/MovieScene/Public/MovieSceneSequencePlayer.h
+++ b/Engine/Source/Runtime/MovieScene/Public/MovieSceneSequencePlayer.h
@@ -548,6 +548,9 @@ protected:
 	MOVIESCENE_API virtual bool CallRemoteFunction(UFunction* Function, void* Parameters, FOutParmRec* OutParms, FFrame* Stack) override;
 	MOVIESCENE_API virtual void PostNetReceive() override;
 	MOVIESCENE_API virtual void BeginDestroy() override;
+#if UE_WITH_IRIS
+	MOVIESCENE_API virtual void RegisterReplicationFragments(UE::Net::FFragmentRegistrationContext& Context, UE::Net::EFragmentRegistrationFlags RegistrationFlags) override;
+#endif
 	/*~ End UObject interface */
 
 	//~ Begin IMovieSceneSequenceTickManagerClient interface
