LevelSequenceActor only supports the subobject list and needs to force it true.

#rnx

diff --git a/Engine/Source/Runtime/LevelSequence/Private/LevelSequenceActor.cpp b/Engine/Source/Runtime/LevelSequence/Private/LevelSequenceActor.cpp
--- a/Engine/Source/Runtime/LevelSequence/Private/LevelSequenceActor.cpp
+++ b/Engine/Source/Runtime/LevelSequence/Private/LevelSequenceActor.cpp
@@ -81,6 +81,7 @@ ALevelSequenceActor::ALevelSequenceActor(const FObjectInitializer& Init)
 
 	bReplicates = true;
 	bReplicatePlayback = false;
+	bReplicateUsingRegisteredSubObjectList = true;
 }
 
 void ALevelSequenceActor::PostInitProperties()
