Sequencer: Linked Anim Instances. Delete links if we can't resolve(by loading) the object referenced.

diff --git a/Engine/Source/Runtime/LevelSequence/Private/AnimSequenceLevelSequenceLink.cpp b/Engine/Source/Runtime/LevelSequence/Private/AnimSequenceLevelSequenceLink.cpp
--- a/Engine/Source/Runtime/LevelSequence/Private/AnimSequenceLevelSequenceLink.cpp
+++ b/Engine/Source/Runtime/LevelSequence/Private/AnimSequenceLevelSequenceLink.cpp
@@ -23,3 +23,12 @@ UAnimSequenceLevelSequenceLink::UAnimSequenceLevelSequenceLink(class FObjectInit
 {
 
 }
+
+void UAnimSequenceLevelSequenceLink::PostLoad()
+{
+	Super::PostLoad();
+	if (PathToLevelSequence.TryLoad() == nullptr)
+	{
+		PathToLevelSequence.Reset();
+	}
+}
diff --git a/Engine/Source/Runtime/LevelSequence/Private/LevelSequenceAnimSequenceLink.cpp b/Engine/Source/Runtime/LevelSequence/Private/LevelSequenceAnimSequenceLink.cpp
--- a/Engine/Source/Runtime/LevelSequence/Private/LevelSequenceAnimSequenceLink.cpp
+++ b/Engine/Source/Runtime/LevelSequence/Private/LevelSequenceAnimSequenceLink.cpp
@@ -23,3 +23,16 @@ ULevelSequenceAnimSequenceLink::ULevelSequenceAnimSequenceLink(class FObjectInit
 {
 
 }
+
+void ULevelSequenceAnimSequenceLink::PostLoad()
+{
+	Super::PostLoad();
+	for (int32 Index = AnimSequenceLinks.Num() - 1; Index >= 0; --Index)
+	{
+		if (AnimSequenceLinks[Index].PathToAnimSequence.TryLoad() == nullptr)
+		{
+			AnimSequenceLinks[Index].PathToAnimSequence.Reset();
+			AnimSequenceLinks.RemoveAt(Index);
+		}
+	}
+}
\ No newline at end of file
diff --git a/Engine/Source/Runtime/LevelSequence/Public/AnimSequenceLevelSequenceLink.h b/Engine/Source/Runtime/LevelSequence/Public/AnimSequenceLevelSequenceLink.h
--- a/Engine/Source/Runtime/LevelSequence/Public/AnimSequenceLevelSequenceLink.h
+++ b/Engine/Source/Runtime/LevelSequence/Public/AnimSequenceLevelSequenceLink.h
@@ -23,4 +23,8 @@ public:
 
 	LEVELSEQUENCE_API void SetLevelSequence(ULevelSequence* InLevelSequence);
 	LEVELSEQUENCE_API ULevelSequence* ResolveLevelSequence();
+
+	//~ Begin UObject Interface
+	virtual void PostLoad() override;
+
 };
diff --git a/Engine/Source/Runtime/LevelSequence/Public/LevelSequenceAnimSequenceLink.h b/Engine/Source/Runtime/LevelSequence/Public/LevelSequenceAnimSequenceLink.h
--- a/Engine/Source/Runtime/LevelSequence/Public/LevelSequenceAnimSequenceLink.h
+++ b/Engine/Source/Runtime/LevelSequence/Public/LevelSequenceAnimSequenceLink.h
@@ -62,4 +62,7 @@ public:
 
 	UPROPERTY(BlueprintReadWrite, Category = Links)
 	TArray< FLevelSequenceAnimSequenceLinkItem> AnimSequenceLinks;
+
+	//~ Begin UObject Interface
+	virtual void PostLoad() override;
 };
