Sequencer: Linked Anim Instances. Delete links if we can't resolve(by loading) the object referenced.

diff --git a/Engine/Source/Runtime/LevelSequence/Private/AnimSequenceLevelSequenceLink.cpp b/Engine/Source/Runtime/LevelSequence/Private/AnimSequenceLevelSequenceLink.cpp
--- a/Engine/Source/Runtime/LevelSequence/Private/AnimSequenceLevelSequenceLink.cpp
+++ b/Engine/Source/Runtime/LevelSequence/Private/AnimSequenceLevelSequenceLink.cpp
@@ -26,9 +26,11 @@ UAnimSequenceLevelSequenceLink::UAnimSequenceLevelSequenceLink(class FObjectInit
 
 void UAnimSequenceLevelSequenceLink::PostLoad()
 {
+#if WITH_EDITOR
 	Super::PostLoad();
 	if (PathToLevelSequence.TryLoad() == nullptr)
 	{
 		PathToLevelSequence.Reset();
 	}
+#endif
 }
diff --git a/Engine/Source/Runtime/LevelSequence/Private/LevelSequenceAnimSequenceLink.cpp b/Engine/Source/Runtime/LevelSequence/Private/LevelSequenceAnimSequenceLink.cpp
--- a/Engine/Source/Runtime/LevelSequence/Private/LevelSequenceAnimSequenceLink.cpp
+++ b/Engine/Source/Runtime/LevelSequence/Private/LevelSequenceAnimSequenceLink.cpp
@@ -27,6 +27,7 @@ ULevelSequenceAnimSequenceLink::ULevelSequenceAnimSequenceLink(class FObjectInit
 void ULevelSequenceAnimSequenceLink::PostLoad()
 {
 	Super::PostLoad();
+#if WITH_EDITOR
 	for (int32 Index = AnimSequenceLinks.Num() - 1; Index >= 0; --Index)
 	{
 		if (AnimSequenceLinks[Index].PathToAnimSequence.TryLoad() == nullptr)
@@ -35,4 +36,5 @@ void ULevelSequenceAnimSequenceLink::PostLoad()
 			AnimSequenceLinks.RemoveAt(Index);
 		}
 	}
+#endif
 }
\ No newline at end of file
