Sequencer: fix crash when adding possessables to a level sequence via scripting

diff --git a/Engine/Source/Runtime/LevelSequence/Private/LevelSequence.cpp b/Engine/Source/Runtime/LevelSequence/Private/LevelSequence.cpp
--- a/Engine/Source/Runtime/LevelSequence/Private/LevelSequence.cpp
+++ b/Engine/Source/Runtime/LevelSequence/Private/LevelSequence.cpp
@@ -593,6 +593,8 @@ void ULevelSequence::OnDirectorRecompiled(UBlueprint* InCompiledBlueprint)
 
 FGuid ULevelSequence::FindOrAddBinding(UObject* InObject)
 {
+	using namespace UE::MovieScene;
+
 	UObject* PlaybackContext = InObject ? InObject->GetWorld() : nullptr;
 	if (!InObject || !PlaybackContext)
 	{
@@ -623,20 +625,15 @@ FGuid ULevelSequence::FindOrAddBinding(UObject* InObject)
 
 	// Perform a potentially slow lookup of every possessable binding in the sequence to see if we already have this
 	{
-		class FTransientPlayer : public IMovieScenePlayer
-		{
-		public:
-			FMovieSceneRootEvaluationTemplateInstance Template;
-			virtual FMovieSceneRootEvaluationTemplateInstance& GetEvaluationTemplate() override { return Template; }
-			virtual void SetViewportSettings(const TMap<FViewportClient*, EMovieSceneViewportParams>& ViewportParamsMap) override {}
-			virtual void GetViewportSettings(TMap<FViewportClient*, EMovieSceneViewportParams>& ViewportParamsMap) const override {}
-			virtual EMovieScenePlayerStatus::Type GetPlaybackStatus() const { return EMovieScenePlayerStatus::Stopped; }
-			virtual void SetPlaybackStatus(EMovieScenePlayerStatus::Type InPlaybackStatus) override {}
-		} Player;
-
-		Player.State.AssignSequence(MovieSceneSequenceID::Root, *this, Player);
-
-		FGuid ExistingID = Player.FindObjectId(*InObject, MovieSceneSequenceID::Root);
+		FSharedPlaybackStateCreateParams CreateParams;
+		CreateParams.PlaybackContext = PlaybackContext;
+		TSharedRef<FSharedPlaybackState> TransientPlaybackState = MakeShared<FSharedPlaybackState>(*this, CreateParams);
+
+		FMovieSceneEvaluationState State;
+		TransientPlaybackState->AddCapabilityRaw(&State);
+		State.AssignSequence(MovieSceneSequenceID::Root, *this, TransientPlaybackState);
+
+		FGuid ExistingID = State.FindObjectId(*InObject, MovieSceneSequenceID::Root, TransientPlaybackState);
 		if (ExistingID.IsValid())
 		{
 			return ExistingID;
diff --git a/Engine/Source/Runtime/MovieScene/Public/EntitySystem/MovieSceneSharedPlaybackState.h b/Engine/Source/Runtime/MovieScene/Public/EntitySystem/MovieSceneSharedPlaybackState.h
--- a/Engine/Source/Runtime/MovieScene/Public/EntitySystem/MovieSceneSharedPlaybackState.h
+++ b/Engine/Source/Runtime/MovieScene/Public/EntitySystem/MovieSceneSharedPlaybackState.h
@@ -56,7 +56,7 @@ struct FSharedPlaybackStateCreateParams
 /**
  * A structure that stores playback state for an entire sequence hierarchy.
  */
-struct FSharedPlaybackState : TSharedFromThis<FSharedPlaybackState>
+struct MOVIESCENE_API FSharedPlaybackState : TSharedFromThis<FSharedPlaybackState>
 {
 public:
 
@@ -89,10 +89,10 @@ public:
 
 	// General utility methods
 
-	MOVIESCENE_API UMovieSceneEntitySystemLinker* GetLinker() const;
+	UMovieSceneEntitySystemLinker* GetLinker() const;
 
-	MOVIESCENE_API const FMovieSceneSequenceHierarchy* GetHierarchy() const;
-	MOVIESCENE_API UMovieSceneSequence* GetSequence(FMovieSceneSequenceIDRef SequenceID) const;
+	const FMovieSceneSequenceHierarchy* GetHierarchy() const;
+	UMovieSceneSequence* GetSequence(FMovieSceneSequenceIDRef SequenceID) const;
 
 public:
 
diff --git a/Engine/Source/Runtime/MovieScene/Public/Evaluation/MovieSceneEvaluationState.h b/Engine/Source/Runtime/MovieScene/Public/Evaluation/MovieSceneEvaluationState.h
--- a/Engine/Source/Runtime/MovieScene/Public/Evaluation/MovieSceneEvaluationState.h
+++ b/Engine/Source/Runtime/MovieScene/Public/Evaluation/MovieSceneEvaluationState.h
@@ -302,7 +302,7 @@ struct FMovieSceneEvaluationState : public UE::MovieScene::IPlaybackCapability
 {
 	using FSharedPlaybackState = UE::MovieScene::FSharedPlaybackState;
 
-	static UE::MovieScene::TPlaybackCapabilityID<FMovieSceneEvaluationState> ID;
+	MOVIESCENE_API static UE::MovieScene::TPlaybackCapabilityID<FMovieSceneEvaluationState> ID;
 
 	/**
 	 * Assign a sequence to a specific ID
@@ -477,8 +477,8 @@ struct FMovieSceneEvaluationState : public UE::MovieScene::IPlaybackCapability
 public:
 
 	/** IPlaybackCapability members */
-	virtual void Initialize(TSharedRef<const FSharedPlaybackState> Owner) override;
-	virtual void OnSubInstanceCreated(TSharedRef<const FSharedPlaybackState> Owner, const UE::MovieScene::FInstanceHandle InstanceHandle) override;
+	MOVIESCENE_API virtual void Initialize(TSharedRef<const FSharedPlaybackState> Owner) override;
+	MOVIESCENE_API virtual void OnSubInstanceCreated(TSharedRef<const FSharedPlaybackState> Owner, const UE::MovieScene::FInstanceHandle InstanceHandle) override;
 
 private:
 
