Sequencer: initialize playback capabilities in InitRootSequenceInstance

diff --git a/Engine/Source/Editor/Sequencer/Private/Sequencer.cpp b/Engine/Source/Editor/Sequencer/Private/Sequencer.cpp
--- a/Engine/Source/Editor/Sequencer/Private/Sequencer.cpp
+++ b/Engine/Source/Editor/Sequencer/Private/Sequencer.cpp
@@ -664,10 +664,7 @@ void FSequencer::InitRootSequenceInstance()
 	TSharedPtr<FSharedPlaybackState> SharedPlaybackState = RootTemplateInstance.GetSharedPlaybackState();
 	if (ensure(SharedPlaybackState.IsValid()))
 	{
-		if (!SharedPlaybackState->HasCapability<FCameraCutPlaybackCapability>())
-		{
-			SharedPlaybackState->AddCapabilityRaw<FCameraCutPlaybackCapability>((FCameraCutPlaybackCapability*)this);
-		}
+		SharedPlaybackState->SetOrAddCapabilityRaw<FCameraCutPlaybackCapability>((FCameraCutPlaybackCapability*)this);
 	}
 }
 
@@ -1121,15 +1118,7 @@ void FSequencer::ResetToNewRootSequence(UMovieSceneSequence& NewSequence)
 
 	AddNodeGroupsCollectionChangedDelegate();
 
-	// Add the camera cut playback capability.
-	FInstanceHandle RootInstanceHandle = RootTemplateInstance.GetRootInstanceHandle();
-	UMovieSceneEntitySystemLinker* Linker = GetEvaluationTemplate().GetEntitySystemLinker();
-	FInstanceRegistry* InstanceRegistry = Linker->GetInstanceRegistry();
-	TSharedRef<FSharedPlaybackState> SharedPlaybackState = InstanceRegistry->GetInstance(RootInstanceHandle).GetSharedPlaybackState();
-	if (!SharedPlaybackState->HasCapability<FCameraCutPlaybackCapability>())
-	{
-		SharedPlaybackState->AddCapabilityRaw<FCameraCutPlaybackCapability>((FCameraCutPlaybackCapability*)this);
-	}
+	InitRootSequenceInstance();
 
 	OnActivateSequenceEvent.Broadcast(ActiveTemplateIDs.Top());
 }
diff --git a/Engine/Source/Runtime/LevelSequence/Private/LevelSequencePlayer.cpp b/Engine/Source/Runtime/LevelSequence/Private/LevelSequencePlayer.cpp
--- a/Engine/Source/Runtime/LevelSequence/Private/LevelSequencePlayer.cpp
+++ b/Engine/Source/Runtime/LevelSequence/Private/LevelSequencePlayer.cpp
@@ -100,20 +100,14 @@ void ULevelSequencePlayer::Initialize(ULevelSequence* InLevelSequence, ULevel* I
 
 	UMovieSceneSequencePlayer::Initialize(InLevelSequence);
 
-	TSharedPtr<FSharedPlaybackState> SharedPlaybackState = RootTemplateInstance.GetSharedPlaybackState();
-	if (SharedPlaybackState)
-	{
-		// The parent player class' root evaluation template may or may not have re-initialized itself.
-		// For instance, if we are given the same sequence asset we already had before, and nothing else
-		// (such as playback context) has changed, no actual re-initialization occurs and we keep the
-		// same shared playback state as before.
-		// That state would already have the spawn register and camera cut capabilies... however, our 
-		// spawn register was just re-created (see a few lines above) so we need to overwrite the 
-		// capability pointer to the new object.
-		// The camera cut capability will stay as a pointer to ourselves, and that's fine.
-		SharedPlaybackState->SetOrAddCapabilityRaw<FMovieSceneSpawnRegister>(SpawnRegister.Get());
-		SharedPlaybackState->SetOrAddCapabilityRaw<FCameraCutPlaybackCapability>((FCameraCutPlaybackCapability*)this);
-	}
+	// The parent player class' root evaluation template may or may not have re-initialized itself.
+	// For instance, if we are given the same sequence asset we already had before, and nothing else
+	// (such as playback context) has changed, no actual re-initialization occurs and we keep the
+	// same shared playback state as before.
+	// That state would already have the spawn register and camera cut capabilies... however, our 
+	// spawn register was just re-created (see a few lines above) so we need to overwrite the
+	// capability pointer to the new object.
+	InitializeLevelSequenceRootInstance(RootTemplateInstance.GetSharedPlaybackState().ToSharedRef());
 }
 
 void ULevelSequencePlayer::SetSourceActorContext(UWorld* InStreamingWorld, FActorContainerID InContainerID, FTopLevelAssetPath InSourceAssetPath)
@@ -250,6 +244,21 @@ void ULevelSequencePlayer::GetEventContexts(UWorld& InWorld, TArray<UObject*>& O
 	}
 }
 
+void ULevelSequencePlayer::InitializeRootInstance(TSharedRef<UE::MovieScene::FSharedPlaybackState> NewSharedPlaybackState)
+{
+	using namespace UE::MovieScene;
+
+	Super::InitializeRootInstance(NewSharedPlaybackState);
+
+	InitializeLevelSequenceRootInstance(NewSharedPlaybackState);
+}
+
+void ULevelSequencePlayer::InitializeLevelSequenceRootInstance(TSharedRef<UE::MovieScene::FSharedPlaybackState> NewSharedPlaybackState)
+{
+	NewSharedPlaybackState->SetOrAddCapabilityRaw<FMovieSceneSpawnRegister>(SpawnRegister.Get());
+	NewSharedPlaybackState->SetOrAddCapabilityRaw<FCameraCutPlaybackCapability>((FCameraCutPlaybackCapability*)this);
+}
+
 void ULevelSequencePlayer::TakeFrameSnapshot(FLevelSequencePlayerSnapshot& OutSnapshot) const
 {
 	if (!ensure(Sequence))
diff --git a/Engine/Source/Runtime/LevelSequence/Public/LevelSequencePlayer.h b/Engine/Source/Runtime/LevelSequence/Public/LevelSequencePlayer.h
--- a/Engine/Source/Runtime/LevelSequence/Public/LevelSequencePlayer.h
+++ b/Engine/Source/Runtime/LevelSequence/Public/LevelSequencePlayer.h
@@ -133,6 +133,7 @@ protected:
 
 	// IMovieScenePlayer interface
 	LEVELSEQUENCE_API virtual void ResolveBoundObjects(const FGuid& InBindingId, FMovieSceneSequenceID SequenceID, UMovieSceneSequence& InSequence, UObject* ResolutionContext, TArray<UObject*, TInlineAllocator<1>>& OutObjects) const override;
+	LEVELSEQUENCE_API virtual void InitializeRootInstance(TSharedRef<UE::MovieScene::FSharedPlaybackState> NewSharedPlaybackState) override;
 
 	//~ UMovieSceneSequencePlayer interface
 	LEVELSEQUENCE_API virtual bool CanPlay() const override;
@@ -161,6 +162,8 @@ private:
 
 	LEVELSEQUENCE_API void EnableCinematicMode(bool bEnable);
 
+	void InitializeLevelSequenceRootInstance(TSharedRef<UE::MovieScene::FSharedPlaybackState> NewSharedPlaybackState);
+
 private:
 
 	/** The world this player will spawn actors in, if needed */
