Sequencer: Fix PostNetReceive not actually initializing the player because it was missing a check for whether the evaluation template is valid or not.

The fix here is to just call InitializePlayer() which is already doing the necessary checks.  This fixes an issue where if the player has been torn down (ie. EndPlay has been called on the Level Sequence Actor), the player will be reinitialized when it is reconstructed.

diff --git a/Engine/Source/Runtime/LevelSequence/Private/LevelSequenceActor.cpp b/Engine/Source/Runtime/LevelSequence/Private/LevelSequenceActor.cpp
--- a/Engine/Source/Runtime/LevelSequence/Private/LevelSequenceActor.cpp
+++ b/Engine/Source/Runtime/LevelSequence/Private/LevelSequenceActor.cpp
@@ -572,10 +572,7 @@ void ALevelSequenceActor::PostNetReceive()
 {
 	Super::PostNetReceive();
 
-	if (LevelSequenceAsset && GetSequencePlayer() && GetSequencePlayer()->GetSequence() != LevelSequenceAsset)
-	{
-		InitializePlayer();
-	}
+	InitializePlayer();
 }
 
 #if WITH_EDITOR
