Sequencer: remove player dependency from director blueprints, and make the playback capability auto-managed.

diff --git a/Engine/Plugins/MovieScene/ActorSequence/Source/ActorSequence/Private/ActorSequence.cpp b/Engine/Plugins/MovieScene/ActorSequence/Source/ActorSequence/Private/ActorSequence.cpp
--- a/Engine/Plugins/MovieScene/ActorSequence/Source/ActorSequence/Private/ActorSequence.cpp
+++ b/Engine/Plugins/MovieScene/ActorSequence/Source/ActorSequence/Private/ActorSequence.cpp
@@ -171,9 +171,9 @@ void UActorSequence::UnbindPossessableObjects(const FGuid& ObjectId)
 	ObjectReferences.RemoveBinding(ObjectId);
 }
 
-UObject* UActorSequence::CreateDirectorInstance(IMovieScenePlayer& Player, FMovieSceneSequenceID SequenceID)
+UObject* UActorSequence::CreateDirectorInstance(TSharedRef<const FSharedPlaybackState> SharedPlaybackState, FMovieSceneSequenceID SequenceID)
 {
-	AActor* Actor = CastChecked<AActor>(Player.GetPlaybackContext(), ECastCheckedType::NullAllowed);
+	AActor* Actor = CastChecked<AActor>(SharedPlaybackState->GetPlaybackContext(), ECastCheckedType::NullAllowed);
 	if (!Actor)
 	{
 		return nullptr;
diff --git a/Engine/Plugins/MovieScene/ActorSequence/Source/ActorSequence/Public/ActorSequence.h b/Engine/Plugins/MovieScene/ActorSequence/Source/ActorSequence/Public/ActorSequence.h
--- a/Engine/Plugins/MovieScene/ActorSequence/Source/ActorSequence/Public/ActorSequence.h
+++ b/Engine/Plugins/MovieScene/ActorSequence/Source/ActorSequence/Public/ActorSequence.h
@@ -28,7 +28,7 @@ public:
 	virtual void UnbindPossessableObjects(const FGuid& ObjectId) override;
 	virtual void UnbindObjects(const FGuid& ObjectId, const TArray<UObject*>& InObjects, UObject* Context) override {}
 	virtual void UnbindInvalidObjects(const FGuid& ObjectId, UObject* Context) override {}
-	virtual UObject* CreateDirectorInstance(IMovieScenePlayer& Player, FMovieSceneSequenceID SequenceID) override;
+	virtual UObject* CreateDirectorInstance(TSharedRef<const FSharedPlaybackState> SharedPlaybackState, FMovieSceneSequenceID SequenceID) override;
 
 #if WITH_EDITOR
 	virtual FText GetDisplayName() const override;
@@ -72,4 +72,4 @@ private:
 	UPROPERTY()
 	bool bHasBeenInitialized;
 #endif
-};
\ No newline at end of file
+};
diff --git a/Engine/Source/Runtime/LevelSequence/Private/LevelSequence.cpp b/Engine/Source/Runtime/LevelSequence/Private/LevelSequence.cpp
--- a/Engine/Source/Runtime/LevelSequence/Private/LevelSequence.cpp
+++ b/Engine/Source/Runtime/LevelSequence/Private/LevelSequence.cpp
@@ -737,10 +737,10 @@ FGuid ULevelSequence::CreateSpawnable(UObject* ObjectToSpawn)
 
 #endif // WITH_EDITOR
 
-UObject* ULevelSequence::CreateDirectorInstance(IMovieScenePlayer& Player, FMovieSceneSequenceID SequenceID)
+UObject* ULevelSequence::CreateDirectorInstance(TSharedRef<const FSharedPlaybackState> SharedPlaybackState, FMovieSceneSequenceID SequenceID)
 {
-	ULevelSequencePlayer* LevelSequencePlayer = Cast<ULevelSequencePlayer>(Player.AsUObject());
-	UObject*              DirectorOuter       = LevelSequencePlayer ? LevelSequencePlayer : Player.GetPlaybackContext();
+	UObject* DirectorOuter = SharedPlaybackState->GetPlaybackContext();
+	IMovieScenePlayer* OptionalPlayer = UE::MovieScene::FPlayerIndexPlaybackCapability::GetPlayer(SharedPlaybackState);
 
 #if WITH_EDITOR
 	if (!UMovieScene::IsTrackClassAllowed(ULevelSequenceDirector::StaticClass()))
@@ -758,10 +758,19 @@ UObject* ULevelSequence::CreateDirectorInstance(IMovieScenePlayer& Player, FMovi
 		DirectorName = MakeUniqueObjectName(DirectorOuter, DirectorClass, *(GetFName().ToString() + TEXT("_Director")));
 #endif
 
+		ULevelSequencePlayer* LevelSequencePlayer = nullptr;
+		if (OptionalPlayer)
+		{
+			LevelSequencePlayer = Cast<ULevelSequencePlayer>(OptionalPlayer->AsUObject());
+		}
+
 		ULevelSequenceDirector* NewDirector = NewObject<ULevelSequenceDirector>(DirectorOuter, DirectorClass, DirectorName, RF_Transient);
-		NewDirector->Player = LevelSequencePlayer;
-		NewDirector->MovieScenePlayerIndex = Player.GetUniqueIndex();
 		NewDirector->SubSequenceID = SequenceID.GetInternalValue();
+		NewDirector->WeakLinker = SharedPlaybackState->GetLinker();
+		NewDirector->InstanceID = SharedPlaybackState->GetRootInstanceHandle().InstanceID;
+		NewDirector->InstanceSerial = SharedPlaybackState->GetRootInstanceHandle().InstanceSerial;
+		NewDirector->Player = LevelSequencePlayer;
+		NewDirector->MovieScenePlayerIndex = OptionalPlayer ? OptionalPlayer->GetUniqueIndex() : INDEX_NONE;
 		NewDirector->OnCreated();
 		return NewDirector;
 	}
diff --git a/Engine/Source/Runtime/LevelSequence/Private/LevelSequenceDirector.cpp b/Engine/Source/Runtime/LevelSequence/Private/LevelSequenceDirector.cpp
--- a/Engine/Source/Runtime/LevelSequence/Private/LevelSequenceDirector.cpp
+++ b/Engine/Source/Runtime/LevelSequence/Private/LevelSequenceDirector.cpp
@@ -8,7 +8,9 @@
 #include "LevelSequencePlayer.h"
 #include "IMovieScenePlayer.h"
 #include "Evaluation/MovieSceneEvaluationTemplateInstance.h"
+#include "Evaluation/MovieSceneSequenceHierarchy.h"
 #include "MovieSceneSequence.h"
+#include "EntitySystem/MovieSceneEntitySystemLinker.h"
 #include "EntitySystem/MovieSceneSequenceInstance.h"
 #include "Evaluation/MovieScenePlayback.h"
 
@@ -23,20 +25,40 @@ UWorld* ULevelSequenceDirector::GetWorld() const
 	return GetTypedOuter<UWorld>();
 }
 
+const UE::MovieScene::FSequenceInstance* ULevelSequenceDirector::FindSequenceInstance() const
+{
+	using namespace UE::MovieScene;
+
+	if (UMovieSceneEntitySystemLinker* Linker = WeakLinker.Get())
+	{
+		const FInstanceHandle InstanceHandle(InstanceID, InstanceSerial);
+		const FInstanceRegistry* InstanceRegistry = Linker->GetInstanceRegistry();
+		if (InstanceRegistry->IsHandleValid(InstanceHandle))
+		{
+			const FSequenceInstance& Instance = InstanceRegistry->GetInstance(InstanceHandle);
+			return &Instance;
+		}
+	}
+	return nullptr;
+}
+
 FQualifiedFrameTime ULevelSequenceDirector::GetRootSequenceTime() const
 {
 	using namespace UE::MovieScene;
 
-	if (IMovieScenePlayer* PlayerInterface = IMovieScenePlayer::Get(static_cast<uint16>(MovieScenePlayerIndex)))
+	if (const FSequenceInstance* Instance = FindSequenceInstance())
 	{
-		FMovieSceneRootEvaluationTemplateInstance& EvaluationTemplate = PlayerInterface->GetEvaluationTemplate();
-		UMovieSceneSequence* RootSequence = EvaluationTemplate.GetRootSequence();
-		const FSequenceInstance* RootSequenceInstance = EvaluationTemplate.FindInstance(MovieSceneSequenceID::Root);
-		if (RootSequenceInstance && RootSequence)
+		TSharedPtr<const FSharedPlaybackState> SharedPlaybackState = Instance->GetSharedPlaybackState();
+		UMovieSceneSequence* RootSequence = SharedPlaybackState->GetRootSequence();
+		if (RootSequence)
 		{
+			const FRootInstanceHandle RootInstanceHandle = Instance->GetRootInstanceHandle();
+			const FInstanceRegistry* InstanceRegistry = SharedPlaybackState->GetLinker()->GetInstanceRegistry();
+			const FSequenceInstance& RootInstance = InstanceRegistry->GetInstance(RootInstanceHandle);
+
 			// Put the qualified frame time into 'display' rate
 			FFrameRate DisplayRate = RootSequence->GetMovieScene()->GetDisplayRate();
-			FMovieSceneContext Context = RootSequenceInstance->GetContext();
+			FMovieSceneContext Context = RootInstance.GetContext();
 
 			FFrameTime DisplayRateTime = ConvertFrameTime(Context.GetTime(), Context.GetFrameRate(), DisplayRate);
 			return FQualifiedFrameTime(DisplayRateTime, DisplayRate);
@@ -49,21 +71,18 @@ FQualifiedFrameTime ULevelSequenceDirector::GetCurrentTime() const
 {
 	using namespace UE::MovieScene;
 
-	if (IMovieScenePlayer* PlayerInterface = IMovieScenePlayer::Get(static_cast<uint16>(MovieScenePlayerIndex)))
+	if (const FSequenceInstance* Instance = FindSequenceInstance())
 	{
-		FMovieSceneRootEvaluationTemplateInstance& EvaluationTemplate = PlayerInterface->GetEvaluationTemplate();
+		FMovieSceneSequenceID ActualSubSequenceID(SubSequenceID);
+		TSharedPtr<const FSharedPlaybackState> SharedPlaybackState = Instance->GetSharedPlaybackState();
+		UMovieSceneSequence* SubSequence = SharedPlaybackState->GetHierarchy()->FindSubSequence(ActualSubSequenceID);
 
-		UMovieSceneSequence* SubSequence = EvaluationTemplate.GetSequence(FMovieSceneSequenceID(SubSequenceID));
-		const FSequenceInstance* SequenceInstance = EvaluationTemplate.FindInstance(FMovieSceneSequenceID(SubSequenceID));
-		if (SequenceInstance && SubSequence)
-		{
-			// Put the qualified frame time into 'display' rate
-			FFrameRate DisplayRate = SubSequence->GetMovieScene()->GetDisplayRate();
-			FMovieSceneContext Context = SequenceInstance->GetContext();
+		// Put the qualified frame time into 'display' rate
+		FFrameRate DisplayRate = SubSequence->GetMovieScene()->GetDisplayRate();
+		FMovieSceneContext Context = Instance->GetContext();
 
-			FFrameTime DisplayRateTime = ConvertFrameTime(Context.GetTime(), Context.GetFrameRate(), DisplayRate);
-			return FQualifiedFrameTime(DisplayRateTime, DisplayRate);
-		}
+		FFrameTime DisplayRateTime = ConvertFrameTime(Context.GetTime(), Context.GetFrameRate(), DisplayRate);
+		return FQualifiedFrameTime(DisplayRateTime, DisplayRate);
 	}
 	return FQualifiedFrameTime(0, FFrameRate());
 }
@@ -71,11 +90,15 @@ FQualifiedFrameTime ULevelSequenceDirector::GetCurrentTime() const
 
 TArray<UObject*> ULevelSequenceDirector::GetBoundObjects(FMovieSceneObjectBindingID ObjectBinding)
 {
+	using namespace UE::MovieScene;
+
 	TArray<UObject*> Objects;
 
-	if (IMovieScenePlayer* PlayerInterface = IMovieScenePlayer::Get(static_cast<uint16>(MovieScenePlayerIndex)))
+	if (const FSequenceInstance* Instance = FindSequenceInstance())
 	{
-		for (TWeakObjectPtr<> WeakObject : ObjectBinding.ResolveBoundObjects(FMovieSceneSequenceID(SubSequenceID), *PlayerInterface))
+		FMovieSceneSequenceID ActualSubSequenceID(SubSequenceID);
+		TSharedPtr<const FSharedPlaybackState> SharedPlaybackState = Instance->GetSharedPlaybackState();
+		for (TWeakObjectPtr<> WeakObject : ObjectBinding.ResolveBoundObjects(ActualSubSequenceID, SharedPlaybackState.ToSharedRef()))
 		{
 			if (UObject* Object = WeakObject.Get())
 			{
@@ -93,9 +116,13 @@ TArray<UObject*> ULevelSequenceDirector::GetBoundObjects(FMovieSceneObjectBindin
 
 UObject* ULevelSequenceDirector::GetBoundObject(FMovieSceneObjectBindingID ObjectBinding)
 {
-	if (IMovieScenePlayer* PlayerInterface = IMovieScenePlayer::Get(static_cast<uint16>(MovieScenePlayerIndex)))
+	using namespace UE::MovieScene;
+
+	if (const FSequenceInstance* Instance = FindSequenceInstance())
 	{
-		for (TWeakObjectPtr<> WeakObject : ObjectBinding.ResolveBoundObjects(FMovieSceneSequenceID(SubSequenceID), *PlayerInterface))
+		FMovieSceneSequenceID ActualSubSequenceID(SubSequenceID);
+		TSharedPtr<const FSharedPlaybackState> SharedPlaybackState = Instance->GetSharedPlaybackState();
+		for (TWeakObjectPtr<> WeakObject : ObjectBinding.ResolveBoundObjects(ActualSubSequenceID, SharedPlaybackState.ToSharedRef()))
 		{
 			if (UObject* Object = WeakObject.Get())
 			{
@@ -113,11 +140,15 @@ UObject* ULevelSequenceDirector::GetBoundObject(FMovieSceneObjectBindingID Objec
 
 TArray<AActor*> ULevelSequenceDirector::GetBoundActors(FMovieSceneObjectBindingID ObjectBinding)
 {
+	using namespace UE::MovieScene;
+
 	TArray<AActor*> Actors;
 
-	if (IMovieScenePlayer* PlayerInterface = IMovieScenePlayer::Get(static_cast<uint16>(MovieScenePlayerIndex)))
+	if (const FSequenceInstance* Instance = FindSequenceInstance())
 	{
-		for (TWeakObjectPtr<> WeakObject : ObjectBinding.ResolveBoundObjects(FMovieSceneSequenceID(SubSequenceID), *PlayerInterface))
+		FMovieSceneSequenceID ActualSubSequenceID(SubSequenceID);
+		TSharedPtr<const FSharedPlaybackState> SharedPlaybackState = Instance->GetSharedPlaybackState();
+		for (TWeakObjectPtr<> WeakObject : ObjectBinding.ResolveBoundObjects(ActualSubSequenceID, SharedPlaybackState.ToSharedRef()))
 		{
 			if (AActor* Actor = Cast<AActor>(WeakObject.Get()))
 			{
@@ -135,9 +166,13 @@ TArray<AActor*> ULevelSequenceDirector::GetBoundActors(FMovieSceneObjectBindingI
 
 AActor* ULevelSequenceDirector::GetBoundActor(FMovieSceneObjectBindingID ObjectBinding)
 {
-	if (IMovieScenePlayer* PlayerInterface = IMovieScenePlayer::Get(static_cast<uint16>(MovieScenePlayerIndex)))
+	using namespace UE::MovieScene;
+
+	if (const FSequenceInstance* Instance = FindSequenceInstance())
 	{
-		for (TWeakObjectPtr<> WeakObject : ObjectBinding.ResolveBoundObjects(FMovieSceneSequenceID(SubSequenceID), *PlayerInterface))
+		FMovieSceneSequenceID ActualSubSequenceID(SubSequenceID);
+		TSharedPtr<const FSharedPlaybackState> SharedPlaybackState = Instance->GetSharedPlaybackState();
+		for (TWeakObjectPtr<> WeakObject : ObjectBinding.ResolveBoundObjects(ActualSubSequenceID, SharedPlaybackState.ToSharedRef()))
 		{
 			if (AActor* Actor = Cast<AActor>(WeakObject.Get()))
 			{
@@ -155,9 +190,13 @@ AActor* ULevelSequenceDirector::GetBoundActor(FMovieSceneObjectBindingID ObjectB
 
 UMovieSceneSequence* ULevelSequenceDirector::GetSequence()
 {
-	if (IMovieScenePlayer* PlayerInterface = IMovieScenePlayer::Get(static_cast<uint16>(MovieScenePlayerIndex)))
+	using namespace UE::MovieScene;
+
+	if (const FSequenceInstance* Instance = FindSequenceInstance())
 	{
-		return PlayerInterface->GetEvaluationTemplate().GetSequence(FMovieSceneSequenceID(SubSequenceID));
+		FMovieSceneSequenceID ActualSubSequenceID(SubSequenceID);
+		TSharedPtr<const FSharedPlaybackState> SharedPlaybackState = Instance->GetSharedPlaybackState();
+		return SharedPlaybackState->GetHierarchy()->FindSubSequence(ActualSubSequenceID);
 	}
 	else
 	{		
diff --git a/Engine/Source/Runtime/LevelSequence/Public/LevelSequence.h b/Engine/Source/Runtime/LevelSequence/Public/LevelSequence.h
--- a/Engine/Source/Runtime/LevelSequence/Public/LevelSequence.h
+++ b/Engine/Source/Runtime/LevelSequence/Public/LevelSequence.h
@@ -53,7 +53,7 @@ public:
 	LEVELSEQUENCE_API virtual bool CanRebindPossessable(const FMovieScenePossessable& InPossessable) const override;
 	LEVELSEQUENCE_API virtual UObject* MakeSpawnableTemplateFromInstance(UObject& InSourceObject, FName ObjectName) override;
 	LEVELSEQUENCE_API virtual bool CanAnimateObject(UObject& InObject) const override;
-	LEVELSEQUENCE_API virtual UObject* CreateDirectorInstance(IMovieScenePlayer& Player, FMovieSceneSequenceID SequenceID) override;
+	LEVELSEQUENCE_API virtual UObject* CreateDirectorInstance(TSharedRef<const FSharedPlaybackState> SharedPlaybackState, FMovieSceneSequenceID SequenceID) override;
 	LEVELSEQUENCE_API virtual const FMovieSceneBindingReferences* GetBindingReferences() const override;
 	LEVELSEQUENCE_API virtual void PostLoad() override;
 #if WITH_EDITORONLY_DATA
diff --git a/Engine/Source/Runtime/LevelSequence/Public/LevelSequenceDirector.h b/Engine/Source/Runtime/LevelSequence/Public/LevelSequenceDirector.h
--- a/Engine/Source/Runtime/LevelSequence/Public/LevelSequenceDirector.h
+++ b/Engine/Source/Runtime/LevelSequence/Public/LevelSequenceDirector.h
@@ -5,6 +5,8 @@
 #include "CoreTypes.h"
 #include "Engine/Blueprint.h"
 #include "Engine/BlueprintGeneratedClass.h"
+#include "EntitySystem/MovieSceneSequenceInstance.h"
+#include "Evaluation/IMovieScenePlaybackCapability.h"
 #include "MovieSceneSequenceID.h"
 #include "MovieSceneObjectBindingID.h"
 #include "Misc/QualifiedFrameTime.h"
@@ -12,6 +14,7 @@
 
 class IMovieScenePlayer;
 class ULevelSequencePlayer;
+class UMovieSceneEntitySystemLinker;
 
 UCLASS(Blueprintable, MinimalAPI)
 class ULevelSequenceDirector : public UObject
@@ -90,14 +93,32 @@ public:
 
 	LEVELSEQUENCE_API virtual UWorld* GetWorld() const override;
 
-	/** Pointer to the player that's playing back this director's sequence. Only valid in game or in PIE/Simulate. */
-	UPROPERTY(BlueprintReadOnly, Category="Cinematics")
-	TObjectPtr<ULevelSequencePlayer> Player;
+private:
+
+	const UE::MovieScene::FSequenceInstance* FindSequenceInstance() const;
+
+public:
 
 	/** The Sequence ID for the sequence this director is playing back within - has to be stored as an int32 so that it is reinstanced correctly*/
 	UPROPERTY()
 	int32 SubSequenceID;
 
+	/** The linker inside which the sequence is evaluating. Only valid in game or in PIE/Simulate. */
+	UPROPERTY()
+	TWeakObjectPtr<UMovieSceneEntitySystemLinker> WeakLinker;
+
+	/** Instance ID of the sequence. Only valid in game or in PIE/Simulate. */
+	UPROPERTY()
+	uint16 InstanceID = (uint16)-1;
+
+	/** Instance serial of the sequence. Only valid in game or in PIE/Simulate. */
+	UPROPERTY()
+	uint16 InstanceSerial = 0;
+
+	/** Pointer to the player that's playing back this director's sequence. Only valid in game or in PIE/Simulate. */
+	UPROPERTY(BlueprintReadOnly, Category="Cinematics")
+	TObjectPtr<ULevelSequencePlayer> Player;
+
 	/** Native player interface index - stored by index so that it can be reinstanced correctly */
 	UPROPERTY()
 	int32 MovieScenePlayerIndex;
diff --git a/Engine/Source/Runtime/MovieScene/Private/Evaluation/MovieSceneEvaluationTemplateInstance.cpp b/Engine/Source/Runtime/MovieScene/Private/Evaluation/MovieSceneEvaluationTemplateInstance.cpp
--- a/Engine/Source/Runtime/MovieScene/Private/Evaluation/MovieSceneEvaluationTemplateInstance.cpp
+++ b/Engine/Source/Runtime/MovieScene/Private/Evaluation/MovieSceneEvaluationTemplateInstance.cpp
@@ -180,8 +180,6 @@ void FMovieSceneRootEvaluationTemplateInstance::Initialize(UMovieSceneSequence&
 		{
 			Player.PreAnimatedState.Initialize(EntitySystemLinker, RootInstanceHandle);
 		}
-
-		ResetDirectorInstances();
 	}
 }
 
@@ -204,14 +202,17 @@ void FMovieSceneRootEvaluationTemplateInstance::EvaluateSynchronousBlocking(FMov
 	}
 }
 
-void FMovieSceneRootEvaluationTemplateInstance::OnFinished()
-{
-	ResetDirectorInstances();
-}
-
 void FMovieSceneRootEvaluationTemplateInstance::ResetDirectorInstances()
 {
-	DirectorInstances.Reset();
+	using namespace UE::MovieScene;
+
+	if (SharedPlaybackState)
+	{
+		if (FSequenceDirectorPlaybackCapability* Cap = SharedPlaybackState->FindCapability<FSequenceDirectorPlaybackCapability>())
+		{
+			Cap->ResetDirectorInstances();
+		}
+	}
 }
 
 bool FMovieSceneRootEvaluationTemplateInstance::IsValid() const
@@ -424,44 +425,16 @@ void FMovieSceneRootEvaluationTemplateInstance::FindEntitiesFromOwner(UObject* O
 
 UObject* FMovieSceneRootEvaluationTemplateInstance::GetOrCreateDirectorInstance(FMovieSceneSequenceIDRef SequenceID, IMovieScenePlayer& Player)
 {
-	UObject* ExistingDirectorInstance = DirectorInstances.FindRef(SequenceID);
-#if WITH_EDITOR
-	if (ExistingDirectorInstance)
-	{
-		// Invalidate our cached director instance if it has been recompiled.
-		UClass* DirectorClass = ExistingDirectorInstance->GetClass();
-		if (!DirectorClass || DirectorClass->HasAnyClassFlags(CLASS_NewerVersionExists))
-		{
-			ExistingDirectorInstance = nullptr;
-		}
-	}
-#endif
-	if (ExistingDirectorInstance)
-	{
-		return ExistingDirectorInstance;
-	}
+	using namespace UE::MovieScene;
 
-	UObject* NewDirectorInstance = nullptr;
-	if (SequenceID == MovieSceneSequenceID::Root)
+	if (SharedPlaybackState)
 	{
-		if (UMovieSceneSequence* Sequence = GetRootSequence())
+		if (FSequenceDirectorPlaybackCapability* Cap = SharedPlaybackState->FindCapability<FSequenceDirectorPlaybackCapability>())
 		{
-			NewDirectorInstance = Sequence->CreateDirectorInstance(Player, SequenceID);
+			return Cap->GetOrCreateDirectorInstance(SharedPlaybackState.ToSharedRef(), SequenceID);
 		}
 	}
-	else if (const FMovieSceneSequenceHierarchy* Hierarchy = GetHierarchy())
-	{
-		const FMovieSceneSubSequenceData* SubData = Hierarchy->FindSubData(SequenceID);
-		check(SubData);
-		NewDirectorInstance = SubData->GetSequence()->CreateDirectorInstance(Player, SequenceID);
-	}
-
-	if (NewDirectorInstance)
-	{
-		DirectorInstances.Add(SequenceID, NewDirectorInstance);
-	}
-
-	return NewDirectorInstance;
+	return nullptr;
 }
 
 void FMovieSceneRootEvaluationTemplateInstance::PlaybackContextChanged(IMovieScenePlayer& Player)
@@ -518,8 +491,6 @@ void FMovieSceneRootEvaluationTemplateInstance::PlaybackContextChanged(IMovieSce
 	SharedPlaybackState = RootInstance.GetSharedPlaybackState();
 	Player.InitializeRootInstance(SharedPlaybackState.ToSharedRef());
 
-	DirectorInstances.Reset();
-
 	Player.PreAnimatedState.Initialize(EntitySystemLinker, RootInstanceHandle);
 	if (bGlobalCapture)
 	{
diff --git a/Engine/Source/Runtime/MovieScene/Private/Evaluation/SequenceDirectorPlaybackCapability.cpp b/Engine/Source/Runtime/MovieScene/Private/Evaluation/SequenceDirectorPlaybackCapability.cpp
--- a/Engine/Source/Runtime/MovieScene/Private/Evaluation/SequenceDirectorPlaybackCapability.cpp
+++ b/Engine/Source/Runtime/MovieScene/Private/Evaluation/SequenceDirectorPlaybackCapability.cpp
@@ -10,7 +10,75 @@
 namespace UE::MovieScene
 {
 
-TPlaybackCapabilityID<ISequenceDirectorPlaybackCapability> ISequenceDirectorPlaybackCapability::ID = TPlaybackCapabilityID<ISequenceDirectorPlaybackCapability>::Register();
+TPlaybackCapabilityID<FSequenceDirectorPlaybackCapability> FSequenceDirectorPlaybackCapability::ID = TPlaybackCapabilityID<FSequenceDirectorPlaybackCapability>::Register();
+
+FString FSequenceDirectorPlaybackCapability::FDirectorInstanceCache::GetReferencerName() const
+{
+	return TEXT("FSequenceDirectorPlaybackCapability");
+}
+
+void FSequenceDirectorPlaybackCapability::FDirectorInstanceCache::AddReferencedObjects(FReferenceCollector& Collector)
+{
+	for (auto& Pair : DirectorInstances)
+	{
+		Collector.AddReferencedObject(Pair.Value);
+	}
+}
+
+void FSequenceDirectorPlaybackCapability::ResetDirectorInstances()
+{
+	if (Cache)
+	{
+		Cache->DirectorInstances.Reset();
+	}
+}
+
+UObject* FSequenceDirectorPlaybackCapability::GetOrCreateDirectorInstance(TSharedRef<const FSharedPlaybackState> SharedPlaybackState, FMovieSceneSequenceIDRef SequenceID)
+{
+	if (!Cache)
+	{
+		Cache = MakeUnique<FDirectorInstanceCache>();
+	}
+
+	UObject* ExistingDirectorInstance = Cache->DirectorInstances.FindRef(SequenceID);
+#if WITH_EDITOR
+	if (ExistingDirectorInstance)
+	{
+		// Invalidate our cached director instance if it has been recompiled.
+		UClass* DirectorClass = ExistingDirectorInstance->GetClass();
+		if (!DirectorClass || DirectorClass->HasAnyClassFlags(CLASS_NewerVersionExists))
+		{
+			ExistingDirectorInstance = nullptr;
+		}
+	}
+#endif
+	if (ExistingDirectorInstance)
+	{
+		return ExistingDirectorInstance;
+	}
+
+	UObject* NewDirectorInstance = nullptr;
+	if (SequenceID == MovieSceneSequenceID::Root)
+	{
+		if (UMovieSceneSequence* Sequence = SharedPlaybackState->GetRootSequence())
+		{
+			NewDirectorInstance = Sequence->CreateDirectorInstance(SharedPlaybackState, SequenceID);
+		}
+	}
+	else if (const FMovieSceneSequenceHierarchy* Hierarchy = SharedPlaybackState->GetHierarchy())
+	{
+		const FMovieSceneSubSequenceData* SubData = Hierarchy->FindSubData(SequenceID);
+		check(SubData);
+		NewDirectorInstance = SubData->GetSequence()->CreateDirectorInstance(SharedPlaybackState, SequenceID);
+	}
+
+	if (NewDirectorInstance)
+	{
+		Cache->DirectorInstances.Add(SequenceID, NewDirectorInstance);
+	}
+
+	return NewDirectorInstance;
+}
 
 }  // namespace UE::MovieScene
 
diff --git a/Engine/Source/Runtime/MovieScene/Private/MovieSceneDynamicBindingInvoker.cpp b/Engine/Source/Runtime/MovieScene/Private/MovieSceneDynamicBindingInvoker.cpp
--- a/Engine/Source/Runtime/MovieScene/Private/MovieSceneDynamicBindingInvoker.cpp
+++ b/Engine/Source/Runtime/MovieScene/Private/MovieSceneDynamicBindingInvoker.cpp
@@ -28,12 +28,13 @@ FMovieSceneDynamicBindingResolveResult FMovieSceneDynamicBindingInvoker::Resolve
 		return FMovieSceneDynamicBindingResolveResult();
 	}
 
-	ISequenceDirectorPlaybackCapability* DirectorCapability = SharedPlaybackState->FindCapability<ISequenceDirectorPlaybackCapability>();
-	if (!ensureMsgf(
-				DirectorCapability, 
-				TEXT("Can't invoke dynamic bindings, the current evaluation context does not have a sequence director playback capability!")))
+	// Auto-add the director playback capability, which is just really a cache for director instances after
+	// they've been created by the sequences in the hierarchy.
+	FSequenceDirectorPlaybackCapability* DirectorCapability = SharedPlaybackState->FindCapability<FSequenceDirectorPlaybackCapability>();
+	if (!DirectorCapability)
 	{
-		return FMovieSceneDynamicBindingResolveResult();
+		TSharedRef<FSharedPlaybackState> MutableState = ConstCastSharedRef<FSharedPlaybackState>(SharedPlaybackState);
+		DirectorCapability = &MutableState->AddCapability<FSequenceDirectorPlaybackCapability>();
 	}
 
 	UObject* DirectorInstance = DirectorCapability->GetOrCreateDirectorInstance(SharedPlaybackState, SequenceID);
diff --git a/Engine/Source/Runtime/MovieScene/Private/MovieScenePlayer.cpp b/Engine/Source/Runtime/MovieScene/Private/MovieScenePlayer.cpp
--- a/Engine/Source/Runtime/MovieScene/Private/MovieScenePlayer.cpp
+++ b/Engine/Source/Runtime/MovieScene/Private/MovieScenePlayer.cpp
@@ -165,12 +165,31 @@ TSharedRef<UE::MovieScene::FSharedPlaybackState> IMovieScenePlayer::GetSharedPla
 
 void IMovieScenePlayer::ResetDirectorInstances()
 {
-	GetEvaluationTemplate().ResetDirectorInstances();
+	using namespace UE::MovieScene;
+
+	TSharedPtr<const FSharedPlaybackState> SharedPlaybackState = FindSharedPlaybackState();
+	if (!SharedPlaybackState)
+	{
+		return;
+	}
+
+	FSequenceDirectorPlaybackCapability* Cap = SharedPlaybackState->FindCapability<FSequenceDirectorPlaybackCapability>();
+	if (Cap)
+	{
+		Cap->ResetDirectorInstances();
+	}
 }
 
 UObject* IMovieScenePlayer::GetOrCreateDirectorInstance(TSharedRef<const UE::MovieScene::FSharedPlaybackState> SharedPlaybackState, FMovieSceneSequenceIDRef SequenceID)
 {
-	return GetEvaluationTemplate().GetOrCreateDirectorInstance(SequenceID, *this);
+	using namespace UE::MovieScene;
+
+	FSequenceDirectorPlaybackCapability* Cap = SharedPlaybackState->FindCapability<FSequenceDirectorPlaybackCapability>();
+	if (Cap)
+	{
+		return Cap->GetOrCreateDirectorInstance(SharedPlaybackState, SequenceID);
+	}
+	return nullptr;
 }
 
 TArray<UObject*> IMovieScenePlayer::GetEventContexts() const
@@ -214,7 +233,6 @@ void IMovieScenePlayer::InitializeRootInstance(TSharedRef<UE::MovieScene::FShare
 	NewSharedPlaybackState->AddCapabilityRaw(&GetSpawnRegister());
 	NewSharedPlaybackState->AddCapabilityRaw((IObjectBindingNotifyPlaybackCapability*)this);
 	NewSharedPlaybackState->AddCapabilityRaw((IStaticBindingOverridesPlaybackCapability*)this);
-	NewSharedPlaybackState->AddCapabilityRaw((ISequenceDirectorPlaybackCapability*)this);
 
 	if (IMovieScenePlaybackClient* PlaybackClient = GetPlaybackClient())
 	{
diff --git a/Engine/Source/Runtime/MovieScene/Private/MovieSceneSequence.cpp b/Engine/Source/Runtime/MovieScene/Private/MovieSceneSequence.cpp
--- a/Engine/Source/Runtime/MovieScene/Private/MovieSceneSequence.cpp
+++ b/Engine/Source/Runtime/MovieScene/Private/MovieSceneSequence.cpp
@@ -399,3 +399,8 @@ FMovieSceneTimecodeSource UMovieSceneSequence::GetEarliestTimecodeSource() const
 	return MovieScene->GetEarliestTimecodeSource();
 }
 
+UObject* UMovieSceneSequence::CreateDirectorInstance(IMovieScenePlayer& Player, FMovieSceneSequenceID SequenceID)
+{
+	return CreateDirectorInstance(Player.GetSharedPlaybackState(), SequenceID);
+}
+
diff --git a/Engine/Source/Runtime/MovieScene/Public/Evaluation/MovieSceneEvaluationTemplateInstance.h b/Engine/Source/Runtime/MovieScene/Public/Evaluation/MovieSceneEvaluationTemplateInstance.h
--- a/Engine/Source/Runtime/MovieScene/Public/Evaluation/MovieSceneEvaluationTemplateInstance.h
+++ b/Engine/Source/Runtime/MovieScene/Public/Evaluation/MovieSceneEvaluationTemplateInstance.h
@@ -65,11 +65,7 @@ public:
 	UE_DEPRECATED(5.1, "Use EvaluateSynchronousBlocking instead.")
 	MOVIESCENE_API void Evaluate(FMovieSceneContext Context, IMovieScenePlayer& Player);
 
-	/**
-	 * Called when this instance has completely finished playing
-	 */
-	void OnFinished();
-
+	UE_DEPRECATED(5.4, "Director instances are now auto-managed via FSequenceDirectorPlaybackCapability")
 	MOVIESCENE_API void ResetDirectorInstances();
 
 	MOVIESCENE_API bool IsValid() const;
@@ -107,6 +103,7 @@ public:
 
 	MOVIESCENE_API UMovieSceneSequence* GetSequence(FMovieSceneSequenceIDRef SequenceID) const;
 
+	UE_DEPRECATED(5.4, "Director instances are now auto-managed via FSequenceDirectorPlaybackCapability")
 	MOVIESCENE_API UObject* GetOrCreateDirectorInstance(FMovieSceneSequenceIDRef SequenceID, IMovieScenePlayer& Player);
 
 	MOVIESCENE_API void PlaybackContextChanged(IMovieScenePlayer& Player);
@@ -144,10 +141,6 @@ private:
 
 	/** The playback state for the hierarchy of sequence instances */
 	TSharedPtr<UE::MovieScene::FSharedPlaybackState> SharedPlaybackState;
-	
-	/** Map of director instances by sequence ID. Kept alive by this map assuming this struct is reference collected */
-	UPROPERTY()
-	TMap<FMovieSceneSequenceID, TObjectPtr<UObject>> DirectorInstances;
 
 	FMovieSceneSequenceID RootID;
 
diff --git a/Engine/Source/Runtime/MovieScene/Public/Evaluation/SequenceDirectorPlaybackCapability.h b/Engine/Source/Runtime/MovieScene/Public/Evaluation/SequenceDirectorPlaybackCapability.h
--- a/Engine/Source/Runtime/MovieScene/Public/Evaluation/SequenceDirectorPlaybackCapability.h
+++ b/Engine/Source/Runtime/MovieScene/Public/Evaluation/SequenceDirectorPlaybackCapability.h
@@ -6,6 +6,7 @@
 #include "Engine/EngineTypes.h"
 #include "Evaluation/MovieScenePlaybackCapabilities.h"
 #include "MovieSceneSequenceID.h"
+#include "UObject/GCObject.h"
 
 namespace UE::MovieScene
 {
@@ -15,18 +16,30 @@ struct FSharedPlaybackState;
 /**
  * Playback capability for sequences that have a director blueprint.
  */
-struct MOVIESCENE_API ISequenceDirectorPlaybackCapability
+struct MOVIESCENE_API FSequenceDirectorPlaybackCapability
 {
 	/** Playback capability ID */
-	static TPlaybackCapabilityID<ISequenceDirectorPlaybackCapability> ID;
-
-	virtual ~ISequenceDirectorPlaybackCapability() {}
+	static TPlaybackCapabilityID<FSequenceDirectorPlaybackCapability> ID;
 
 	/** Remove all director blueprint instances */
-	virtual void ResetDirectorInstances() = 0;
+	void ResetDirectorInstances();
 
 	/** Gets a new or existing director blueprint instance for the given root or sub sequence */
-	virtual UObject* GetOrCreateDirectorInstance(TSharedRef<const FSharedPlaybackState> SharedPlaybackState, FMovieSceneSequenceIDRef SequenceID) = 0;
+	UObject* GetOrCreateDirectorInstance(TSharedRef<const FSharedPlaybackState> SharedPlaybackState, FMovieSceneSequenceIDRef SequenceID);
+
+private:
+
+	// The actual cache of BP instances is stored in a heap-allocated object
+	// because it neeeds to be an FGCObject to keep those instances alive, and
+	// an FGCObject isn't relocatable so we can't put it in a playback 
+	// capabailities container.
+	struct FDirectorInstanceCache : FGCObject
+	{
+		virtual FString GetReferencerName() const override;
+		virtual void AddReferencedObjects(FReferenceCollector& Collector) override;
+		TSortedMap<FMovieSceneSequenceID, TObjectPtr<UObject>> DirectorInstances;
+	};
+	TUniquePtr<FDirectorInstanceCache> Cache;
 };
 
 }  // namespace UE::MovieScene
diff --git a/Engine/Source/Runtime/MovieScene/Public/IMovieScenePlayer.h b/Engine/Source/Runtime/MovieScene/Public/IMovieScenePlayer.h
--- a/Engine/Source/Runtime/MovieScene/Public/IMovieScenePlayer.h
+++ b/Engine/Source/Runtime/MovieScene/Public/IMovieScenePlayer.h
@@ -110,7 +110,6 @@ using EMovieSceneCameraCutParams = FMovieSceneCameraCutParams;
 class IMovieScenePlayer 
 	: public UE::MovieScene::IObjectBindingNotifyPlaybackCapability
 	, public UE::MovieScene::IStaticBindingOverridesPlaybackCapability
-	, public UE::MovieScene::ISequenceDirectorPlaybackCapability
 {
 public:
 	MOVIESCENE_API IMovieScenePlayer();
@@ -222,12 +221,14 @@ public:
 	/**
 	 * Remove all director blueprint instances
 	 */
-	MOVIESCENE_API virtual void ResetDirectorInstances() override;
+	UE_DEPRECATED(5.4, "Director instances are now automanaged via FSequenceDirectorPlaybackCapability")
+	void ResetDirectorInstances();
 
 	/**
 	 * Gets a new or existing director blueprint instance for the given root or sub sequence
 	 */
-	MOVIESCENE_API virtual UObject* GetOrCreateDirectorInstance(TSharedRef<const UE::MovieScene::FSharedPlaybackState> SharedPlaybackState, FMovieSceneSequenceIDRef SequenceID) override;
+	UE_DEPRECATED(5.4, "Director instances are now automanaged via FSequenceDirectorPlaybackCapability")
+	UObject* GetOrCreateDirectorInstance(TSharedRef<const UE::MovieScene::FSharedPlaybackState> SharedPlaybackState, FMovieSceneSequenceIDRef SequenceID);
 
 	/**
 	 * Called to initialize the flag structure that denotes what functions need to be called on this updater
diff --git a/Engine/Source/Runtime/MovieScene/Public/MovieSceneSequence.h b/Engine/Source/Runtime/MovieScene/Public/MovieSceneSequence.h
--- a/Engine/Source/Runtime/MovieScene/Public/MovieSceneSequence.h
+++ b/Engine/Source/Runtime/MovieScene/Public/MovieSceneSequence.h
@@ -42,6 +42,11 @@ struct FMovieSceneTimecodeSource;
 struct FUniversalObjectLocator;
 struct FMovieSceneBindingReferences;
 
+namespace UE::MovieScene
+{
+	struct FSharedPlaybackState;
+}
+
 enum class ETrackSupport
 {
 	/** This track is not supported */
@@ -62,6 +67,8 @@ class UMovieSceneSequence
 {
 public:
 
+	using FSharedPlaybackState = UE::MovieScene::FSharedPlaybackState;
+
 	GENERATED_BODY()
 
 	MOVIESCENE_API UMovieSceneSequence(const FObjectInitializer& Init);
@@ -259,7 +266,10 @@ public:
 	/**
 	 * Called to retrieve or construct a director instance to be used for the specified player
 	 */
-	virtual UObject* CreateDirectorInstance(IMovieScenePlayer& Player, FMovieSceneSequenceID SequenceID) { return nullptr; }
+	virtual UObject* CreateDirectorInstance(TSharedRef<const FSharedPlaybackState> SharedPlaybackState, FMovieSceneSequenceID SequenceID) { return nullptr; }
+
+	UE_DEPRECATED(5.4, "Please use the version that takes a SharedPlaybackState")
+	UObject* CreateDirectorInstance(IMovieScenePlayer& Player, FMovieSceneSequenceID SequenceID);
 
 	MOVIESCENE_API virtual EMovieSceneServerClientMask OverrideNetworkMask(EMovieSceneServerClientMask InDefaultMask) const;
 
diff --git a/Engine/Source/Runtime/MovieSceneTracks/Private/Systems/MovieSceneEventSystems.cpp b/Engine/Source/Runtime/MovieSceneTracks/Private/Systems/MovieSceneEventSystems.cpp
--- a/Engine/Source/Runtime/MovieSceneTracks/Private/Systems/MovieSceneEventSystems.cpp
+++ b/Engine/Source/Runtime/MovieSceneTracks/Private/Systems/MovieSceneEventSystems.cpp
@@ -16,6 +16,7 @@
 #include "IMovieScenePlayer.h"
 #include "MovieSceneSequence.h"
 
+#include "Templates/SharedPointer.h"
 #include "Templates/SubclassOf.h"
 #include "Engine/World.h"
 #include "Engine/LevelScriptActor.h"
@@ -141,10 +142,13 @@ void UMovieSceneEventSystem::TriggerEvents(TArrayView<const FMovieSceneEventTrig
 	UE_LOG(LogMovieScene, VeryVerbose, TEXT("%s: Triggering %d events"), *SequenceName, Events.Num());
 #endif
 
-	ISequenceDirectorPlaybackCapability* DirectorCapability = SharedPlaybackState->FindCapability<ISequenceDirectorPlaybackCapability>();
-	if (!ensureMsgf(DirectorCapability, TEXT("Can't trigger events on %s: no director blueprint capability found"), *SequenceName))
+	// Auto-add the director playback capability, which is just really a cache for director instances after
+	// they've been created by the sequences in the hierarchy.
+	FSequenceDirectorPlaybackCapability* DirectorCapability = SharedPlaybackState->FindCapability<FSequenceDirectorPlaybackCapability>();
+	if (!DirectorCapability)
 	{
-		return;
+		TSharedRef<FSharedPlaybackState> MutableState = ConstCastSharedRef<FSharedPlaybackState>(SharedPlaybackState);
+		DirectorCapability = &MutableState->AddCapability<FSequenceDirectorPlaybackCapability>();
 	}
 
 	TArray<UObject*> GlobalContexts;
diff --git a/Engine/Source/Runtime/UMG/Private/Animation/UMGSequencePlayer.cpp b/Engine/Source/Runtime/UMG/Private/Animation/UMGSequencePlayer.cpp
--- a/Engine/Source/Runtime/UMG/Private/Animation/UMGSequencePlayer.cpp
+++ b/Engine/Source/Runtime/UMG/Private/Animation/UMGSequencePlayer.cpp
@@ -434,10 +434,6 @@ void UUMGSequencePlayer::HandleLatentStop()
 
 	PlayerStatus = EMovieScenePlayerStatus::Stopped;
 
-	// This isn't really necessary because UMG animations cannot have director instances
-	// but it is here for completeness
-	RootTemplateInstance.ResetDirectorInstances();
-
 	if (bRestoreState)
 	{
 		RestorePreAnimatedState();
diff --git a/Engine/Source/Runtime/UMG/Private/Animation/WidgetAnimation.cpp b/Engine/Source/Runtime/UMG/Private/Animation/WidgetAnimation.cpp
--- a/Engine/Source/Runtime/UMG/Private/Animation/WidgetAnimation.cpp
+++ b/Engine/Source/Runtime/UMG/Private/Animation/WidgetAnimation.cpp
@@ -259,10 +259,10 @@ UMovieScene* UWidgetAnimation::GetMovieScene() const
 	return MovieScene;
 }
 
-UObject* UWidgetAnimation::CreateDirectorInstance(IMovieScenePlayer& Player, FMovieSceneSequenceID SequenceID)
+UObject* UWidgetAnimation::CreateDirectorInstance(TSharedRef<const FSharedPlaybackState> SharedPlaybackState, FMovieSceneSequenceID SequenceID)
 {
 	// Widget animations do not create separate director instances, but just re-use the UUserWidget from the playback context
-	UUserWidget* WidgetContext = CastChecked<UUserWidget>(Player.GetPlaybackContext());
+	UUserWidget* WidgetContext = CastChecked<UUserWidget>(SharedPlaybackState->GetPlaybackContext());
 	return WidgetContext;
 }
 
diff --git a/Engine/Source/Runtime/UMG/Public/Animation/WidgetAnimation.h b/Engine/Source/Runtime/UMG/Public/Animation/WidgetAnimation.h
--- a/Engine/Source/Runtime/UMG/Public/Animation/WidgetAnimation.h
+++ b/Engine/Source/Runtime/UMG/Public/Animation/WidgetAnimation.h
@@ -98,7 +98,7 @@ public:
 	virtual void UnbindObjects(const FGuid& ObjectId, const TArray<UObject*>& InObjects, UObject* InContext) override {}
 	virtual void UnbindInvalidObjects(const FGuid& ObjectId, UObject* InContext) override {}
 	virtual void LocateBoundObjects(const FGuid& ObjectId, UObject* Context, TArray<UObject*, TInlineAllocator<1>>& OutObjects) const override;
-	virtual UObject* CreateDirectorInstance(IMovieScenePlayer& Player, FMovieSceneSequenceID SequenceID) override;
+	virtual UObject* CreateDirectorInstance(TSharedRef<const FSharedPlaybackState> SharedPlaybackState, FMovieSceneSequenceID SequenceID) override;
 #if WITH_EDITOR
 	virtual ETrackSupport IsTrackSupported(TSubclassOf<class UMovieSceneTrack> InTrackClass) const override;
 #endif
