Sequencer: fix level director blueprint not getting root sequence info correctly

The APIs to get sequences in the sequence hierarchy only work for sub-sequences, not root sequences. Root sequences need to be checked separately...

diff --git a/Engine/Source/Runtime/LevelSequence/Private/LevelSequenceDirector.cpp b/Engine/Source/Runtime/LevelSequence/Private/LevelSequenceDirector.cpp
--- a/Engine/Source/Runtime/LevelSequence/Private/LevelSequenceDirector.cpp
+++ b/Engine/Source/Runtime/LevelSequence/Private/LevelSequenceDirector.cpp
@@ -73,12 +73,17 @@ FQualifiedFrameTime ULevelSequenceDirector::GetCurrentTime() const
 
 	if (const FSequenceInstance* Instance = FindSequenceInstance())
 	{
-		FMovieSceneSequenceID ActualSubSequenceID(SubSequenceID);
 		TSharedPtr<const FSharedPlaybackState> SharedPlaybackState = Instance->GetSharedPlaybackState();
-		UMovieSceneSequence* SubSequence = SharedPlaybackState->GetHierarchy()->FindSubSequence(ActualSubSequenceID);
+		UMovieSceneSequence* Sequence = SharedPlaybackState->GetRootSequence();
+
+		const FMovieSceneSequenceID ActualSubSequenceID(SubSequenceID);
+		if (ActualSubSequenceID != MovieSceneSequenceID::Root)
+		{
+			Sequence = SharedPlaybackState->GetHierarchy()->FindSubSequence(ActualSubSequenceID);
+		}
 
 		// Put the qualified frame time into 'display' rate
-		FFrameRate DisplayRate = SubSequence->GetMovieScene()->GetDisplayRate();
+		FFrameRate DisplayRate = Sequence->GetMovieScene()->GetDisplayRate();
 		FMovieSceneContext Context = Instance->GetContext();
 
 		FFrameTime DisplayRateTime = ConvertFrameTime(Context.GetTime(), Context.GetFrameRate(), DisplayRate);
@@ -196,7 +201,14 @@ UMovieSceneSequence* ULevelSequenceDirector::GetSequence()
 	{
 		FMovieSceneSequenceID ActualSubSequenceID(SubSequenceID);
 		TSharedPtr<const FSharedPlaybackState> SharedPlaybackState = Instance->GetSharedPlaybackState();
-		return SharedPlaybackState->GetHierarchy()->FindSubSequence(ActualSubSequenceID);
+		if (ActualSubSequenceID == MovieSceneSequenceID::Root)
+		{
+			return SharedPlaybackState->GetRootSequence();
+		}
+		else
+		{
+			return SharedPlaybackState->GetHierarchy()->FindSubSequence(ActualSubSequenceID);
+		}
 	}
 	else
 	{		
