Sequencer: Fixed cleanup of invalid bindings by moving this logic from PostInitProperties to PostLoad

The data is not deserialized yet in PostInitProperties so this would have never done anything previously.

diff --git a/Engine/Source/Runtime/LevelSequence/Private/LevelSequence.cpp b/Engine/Source/Runtime/LevelSequence/Private/LevelSequence.cpp
--- a/Engine/Source/Runtime/LevelSequence/Private/LevelSequence.cpp
+++ b/Engine/Source/Runtime/LevelSequence/Private/LevelSequence.cpp
@@ -313,6 +313,23 @@ void ULevelSequence::PostLoad()
 	Super::PostLoad();
 
 #if WITH_EDITOR
+	if (MovieScene)
+	{
+		// Remove any invalid object bindings. This was moved from PostInitProperties
+		//   because it has to happen after the asset has actually been serialized.
+		TSet<FGuid> ValidObjectBindings;
+		for (int32 Index = 0; Index < MovieScene->GetSpawnableCount(); ++Index)
+		{
+			ValidObjectBindings.Add(MovieScene->GetSpawnable(Index).GetGuid());
+		}
+		for (int32 Index = 0; Index < MovieScene->GetPossessableCount(); ++Index)
+		{
+			ValidObjectBindings.Add(MovieScene->GetPossessable(Index).GetGuid());
+		}
+
+		BindingReferences.RemoveInvalidBindings(ValidObjectBindings);
+	}
+
 	if (!DirectorBlueprint)
 	{
 		UBlueprint* PhantomDirector = FindObject<UBlueprint>(this, TEXT("SequenceDirector"));
@@ -410,24 +427,6 @@ void ULevelSequence::DeclareConstructClasses(TArray<FTopLevelAssetPath>& OutCons
 void ULevelSequence::PostInitProperties()
 {
 	Super::PostInitProperties();
-
-#if WITH_EDITOR
-	if (MovieScene)
-	{
-		// Remove any invalid object bindings
-		TSet<FGuid> ValidObjectBindings;
-		for (int32 Index = 0; Index < MovieScene->GetSpawnableCount(); ++Index)
-		{
-			ValidObjectBindings.Add(MovieScene->GetSpawnable(Index).GetGuid());
-		}
-		for (int32 Index = 0; Index < MovieScene->GetPossessableCount(); ++Index)
-		{
-			ValidObjectBindings.Add(MovieScene->GetPossessable(Index).GetGuid());
-		}
-
-		BindingReferences.RemoveInvalidBindings(ValidObjectBindings);
-	}
-#endif
 }
 
 bool ULevelSequence::Rename(const TCHAR* NewName, UObject* NewOuter, ERenameFlags Flags)
