Sequencer: Back out unnecessary change that removes spaces and makes a unique object name because it makes the spawnable no longer match the spawnable name in the level sequence.

#rnx

diff --git a/Engine/Source/Runtime/LevelSequence/Private/LevelSequenceActorSpawner.cpp b/Engine/Source/Runtime/LevelSequence/Private/LevelSequenceActorSpawner.cpp
--- a/Engine/Source/Runtime/LevelSequence/Private/LevelSequenceActorSpawner.cpp
+++ b/Engine/Source/Runtime/LevelSequence/Private/LevelSequenceActorSpawner.cpp
@@ -143,8 +143,7 @@ UObject* FLevelSequenceActorSpawner::SpawnObject(FMovieSceneSpawnable& Spawnable
 	// Spawn the puppet actor
 	FActorSpawnParameters SpawnInfo;
 	{
-		SpawnInfo.Name = *SpawnName.ToString().Replace(TEXT(" "), TEXT("_"));
-		SpawnInfo.Name = MakeUniqueObjectName(WorldContext->PersistentLevel, ObjectTemplate->GetClass(), SpawnInfo.Name);
+		SpawnInfo.Name = SpawnName;
 		SpawnInfo.ObjectFlags = ObjectFlags;
 		SpawnInfo.SpawnCollisionHandlingOverride = ESpawnActorCollisionHandlingMethod::AlwaysSpawn;
 		// @todo: Spawning with a non-CDO template is fraught with issues
