Fixed a few different crashes that were a result of a replicated LevelSequenceActor being recreated on client, but using a SequencePlayer from a previous, destroyed copy of that actor that hadn't been garbage-collected yet.

The sequence player is now explicitly marked as garbage to prevent finding a matching valid object in the net package map.

diff --git a/Engine/Source/Runtime/LevelSequence/Private/LevelSequenceActor.cpp b/Engine/Source/Runtime/LevelSequence/Private/LevelSequenceActor.cpp
--- a/Engine/Source/Runtime/LevelSequence/Private/LevelSequenceActor.cpp
+++ b/Engine/Source/Runtime/LevelSequence/Private/LevelSequenceActor.cpp
@@ -25,13 +25,22 @@
 	#include "SceneOutlinerFilters.h"
 #endif
 
-bool GLevelSequenceActor_InvalidBindingTagWarnings = true;
-FAutoConsoleVariableRef CVarLevelSequenceActor_InvalidBindingTagWarnings(
-	TEXT("LevelSequence.InvalidBindingTagWarnings"),
-	GLevelSequenceActor_InvalidBindingTagWarnings,
-	TEXT("Whether to emit a warning when invalid object binding tags are used to override bindings or not.\n"),
-	ECVF_Default
-);
+namespace LevelSequenceActorCVars
+{
+	static bool bInvalidBindingTagWarnings = true;
+	static FAutoConsoleVariableRef CVarInvalidBindingTagWarnings(
+		TEXT("LevelSequence.InvalidBindingTagWarnings"),
+		bInvalidBindingTagWarnings,
+		TEXT("Whether to emit a warning when invalid object binding tags are used to override bindings or not.\n"),
+		ECVF_Default);
+
+	static bool bMarkSequencePlayerAsGarbageOnDestroy = true;
+	static FAutoConsoleVariableRef CVarMarkSequencePlayerAsGarbageOnDestroy(
+		TEXT("LevelSequence.MarkSequencePlayerAsGarbageOnDestroy"),
+		bMarkSequencePlayerAsGarbageOnDestroy,
+		TEXT("Whether to flag the sequence player object as garbage when the actor is being destroyed"),
+		ECVF_Default);
+}
 
 ALevelSequenceActor::ALevelSequenceActor(const FObjectInitializer& Init)
 	: Super(Init)
@@ -240,20 +249,29 @@ void ALevelSequenceActor::BeginPlay()
 
 void ALevelSequenceActor::EndPlay(const EEndPlayReason::Type EndPlayReason)
 {
-	if (GetSequencePlayer())
+	if (ULevelSequencePlayer* Player = GetSequencePlayer())
 	{
-		RemoveReplicatedSubObject(GetSequencePlayer());
+		RemoveReplicatedSubObject(Player);
 
 		// Stop may modify a lot of actor state so it needs to be called
 		// during EndPlay (when Actors + World are still valid) instead
 		// of waiting for the UObject to be destroyed by GC.
-		GetSequencePlayer()->Stop();
+		Player->Stop();
 
-		GetSequencePlayer()->OnPlay.RemoveAll(this);
-		GetSequencePlayer()->OnPlayReverse.RemoveAll(this);
-		GetSequencePlayer()->OnStop.RemoveAll(this);
+		Player->OnPlay.RemoveAll(this);
+		Player->OnPlayReverse.RemoveAll(this);
+		Player->OnStop.RemoveAll(this);
 
-		GetSequencePlayer()->TearDown();
+		Player->TearDown();
+
+		// This actor may be being destroyed due to leaving net-relevancy, in which case we need to explicitly
+		// mark the sub-object as garbage. Otherwise, re-entering relevancy will recreate the actor on the client
+		// but may find and assign the existing yet-un-GC'd player sub-object from the previous actor instance.
+		// Actor sub-objects may be automatically marked garbage some day, but for now we take care of it manually.
+		if (LevelSequenceActorCVars::bMarkSequencePlayerAsGarbageOnDestroy && (EndPlayReason == EEndPlayReason::Destroyed))
+		{
+			Player->MarkAsGarbage();
+		}
 	}
 
 	Super::EndPlay(EndPlayReason);
@@ -436,7 +454,7 @@ void ALevelSequenceActor::SetBindingByTag(FName BindingTag, const TArray<AActor*
 			SetBinding(BindingID, Actors, bAllowBindingsFromAsset);
 		}
 	}
-	else if (GLevelSequenceActor_InvalidBindingTagWarnings)
+	else if (LevelSequenceActorCVars::bInvalidBindingTagWarnings)
 	{
 		FMessageLog("PIE")
 			.Warning(FText::Format(NSLOCTEXT("LevelSequenceActor", "SetBindingByTag", "Sequence did not contain any bindings with the tag '{0}'"), FText::FromName(BindingTag)))
@@ -474,7 +492,7 @@ void ALevelSequenceActor::AddBindingByTag(FName BindingTag, AActor* Actor, bool
 			AddBinding(BindingID, Actor, bAllowBindingsFromAsset);
 		}
 	}
-	else if (GLevelSequenceActor_InvalidBindingTagWarnings)
+	else if (LevelSequenceActorCVars::bInvalidBindingTagWarnings)
 	{
 		FMessageLog("PIE")
 			.Warning(FText::Format(NSLOCTEXT("LevelSequenceActor", "AddBindingByTag", "Sequence did not contain any bindings with the tag '{0}'"), FText::FromName(BindingTag)))
@@ -512,7 +530,7 @@ void ALevelSequenceActor::RemoveBindingByTag(FName BindingTag, AActor* Actor)
 			RemoveBinding(BindingID, Actor);
 		}
 	}
-	else if (GLevelSequenceActor_InvalidBindingTagWarnings)
+	else if (LevelSequenceActorCVars::bInvalidBindingTagWarnings)
 	{
 		FMessageLog("PIE")
 			.Warning(FText::Format(NSLOCTEXT("LevelSequenceActor", "RemoveBindingByTag", "Sequence did not contain any bindings with the tag '{0}'"), FText::FromName(BindingTag)))
