Sequencer: Initialize PlaybackSettings in InitializePlayer()

Playback settings weren't always respected when playing a sequence (previously they were only being set in PostLoad and PostInitProperties)

diff --git a/Engine/Source/Runtime/LevelSequence/Private/LevelSequenceActor.cpp b/Engine/Source/Runtime/LevelSequence/Private/LevelSequenceActor.cpp
--- a/Engine/Source/Runtime/LevelSequence/Private/LevelSequenceActor.cpp
+++ b/Engine/Source/Runtime/LevelSequence/Private/LevelSequenceActor.cpp
@@ -365,6 +365,7 @@ void ALevelSequenceActor::InitializePlayer()
 		// Level sequence is already loaded. Initialize the player if it's not already initialized with this sequence
 		if (LevelSequenceAsset != GetSequencePlayer()->GetSequence() || GetSequencePlayer()->GetEvaluationTemplate().GetRunner() == nullptr)
 		{
+			GetSequencePlayer()->SetPlaybackSettings(PlaybackSettings);
 			GetSequencePlayer()->Initialize(LevelSequenceAsset, GetLevel(), CameraSettings);
 		}
 	}
