Fix a few AutoRTFM crashes

diff --git a/Engine/Source/Runtime/LevelSequence/LevelSequence.Build.cs b/Engine/Source/Runtime/LevelSequence/LevelSequence.Build.cs
--- a/Engine/Source/Runtime/LevelSequence/LevelSequence.Build.cs
+++ b/Engine/Source/Runtime/LevelSequence/LevelSequence.Build.cs
@@ -43,5 +43,7 @@ public class LevelSequence : ModuleRules
 				}
 			);
 		}
+
+		bAllowAutoRTFMInstrumentation = true;
 	}
 }
