Sequencer: Set the cached camera component on the player

Accidentally removed in 28093481

diff --git a/Engine/Source/Runtime/LevelSequence/Private/LevelSequencePlayer.cpp b/Engine/Source/Runtime/LevelSequence/Private/LevelSequencePlayer.cpp
--- a/Engine/Source/Runtime/LevelSequence/Private/LevelSequencePlayer.cpp
+++ b/Engine/Source/Runtime/LevelSequence/Private/LevelSequencePlayer.cpp
@@ -231,6 +231,8 @@ TOptional<EAspectRatioAxisConstraint> ULevelSequencePlayer::GetAspectRatioAxisCo
 
 void ULevelSequencePlayer::OnCameraCutUpdated(const UE::MovieScene::FOnCameraCutUpdatedParams& Params)
 {
+	CachedCameraComponent = Params.ViewTargetCamera;
+
 	if (OnCameraCut.IsBound())
 	{
 		OnCameraCut.Broadcast(Params.ViewTargetCamera);
@@ -421,4 +423,3 @@ void ULevelSequencePlayer::RewindForReplay()
 	NetSyncProps.LastKnownNumLoops = 0;
 	NetSyncProps.LastKnownSerialNumber = 0;
 }
-
