Sequencer- Dynamic Binding Refactor

diff --git a/Engine/Source/Editor/UMGEditor/Private/UMGEditorModule.cpp b/Engine/Source/Editor/UMGEditor/Private/UMGEditorModule.cpp
--- a/Engine/Source/Editor/UMGEditor/Private/UMGEditorModule.cpp
+++ b/Engine/Source/Editor/UMGEditor/Private/UMGEditorModule.cpp
@@ -44,6 +44,8 @@
 #include "WidgetBlueprintThumbnailRenderer.h"
 #include "Customizations/WidgetThumbnailCustomization.h"
 #include "Widgets/SBindWidgetView.h"
+#include "MovieSceneDynamicBindingUtils.h"
+#include "Kismet2/KismetEditorUtilities.h"
 
 #define LOCTEXT_NAMESPACE "UMG"
 
@@ -135,6 +137,9 @@ public:
 		FEdGraphUtilities::RegisterVisualPinFactory(GraphPanelPinFactory);
 
 		CVarThumbnailRenderEnable->AsVariable()->SetOnChangedCallback(FConsoleVariableDelegate::CreateStatic(&FUMGEditorModule::ThumbnailRenderingEnabled));
+
+		FixupDynamicBindingPayloadParameterNameHandle = UMovieScene::FixupDynamicBindingPayloadParameterNameEvent.AddStatic(FixupPayloadParameterNameForDynamicBinding);
+		FixupWidgetDynamicBindingsHandle = UWidgetAnimation::FixupWidgetDynamicBindingsEvent.AddStatic(FixupWidgetDynamicBindings);
 	}
 
 	/** Called before the module is unloaded, right before the module object is destroyed. */
@@ -215,6 +220,9 @@ public:
 		//	SettingsModule->UnregisterSettings("Editor", "ContentEditors", "WidgetDesigner");
 		//	SettingsModule->UnregisterSettings("Project", "Editor", "UMGEditor");
 		//}
+
+		UMovieScene::FixupDynamicBindingPayloadParameterNameEvent.Remove(FixupDynamicBindingPayloadParameterNameHandle);
+		UWidgetAnimation::FixupWidgetDynamicBindingsEvent.Remove(FixupWidgetDynamicBindingsHandle);
 	}
 
 	/** Gets the extensibility managers for outside entities to extend gui page editor's menus and toolbars */
@@ -434,6 +442,58 @@ private:
 		}
 	}
 
+	static void FixupPayloadParameterNameForDynamicBinding(UMovieScene* MovieScene, UK2Node* InNode, FName OldPinName, FName NewPinName)
+	{
+		using namespace UE::MovieScene;
+
+		check(MovieScene);
+
+		auto FixupPayloadParameterName = [InNode, OldPinName, NewPinName](FMovieSceneDynamicBinding& DynamicBinding)
+		{
+			if (DynamicBinding.WeakEndpoint.Get() == InNode)
+			{
+				if (FMovieSceneDynamicBindingPayloadVariable* Variable = DynamicBinding.PayloadVariables.Find(OldPinName))
+				{
+					DynamicBinding.PayloadVariables.Add(NewPinName, MoveTemp(*Variable));
+					DynamicBinding.PayloadVariables.Remove(OldPinName);
+				}
+			}
+		};
+
+		UMovieSceneSequence* ThisSequence = MovieScene->GetTypedOuter<UMovieSceneSequence>();
+		TSharedRef<UE::MovieScene::FSharedPlaybackState> TransientPlaybackState = MovieSceneHelpers::CreateTransientSharedPlaybackState(GEditor->GetEditorWorldContext().World(), ThisSequence);
+
+		if (UWidgetAnimation* WidgetAnimation = MovieScene->GetTypedOuter<UWidgetAnimation>())
+		{
+			for (FWidgetAnimationBinding& WidgetAnimationBinding : WidgetAnimation->AnimationBindings)
+			{
+				FixupPayloadParameterName(WidgetAnimationBinding.DynamicBinding);
+			}
+		}
+	}
+
+	static void FixupWidgetDynamicBindings(UWidgetAnimation* WidgetAnimation)
+	{
+		if (WidgetAnimation)
+		{
+			FMovieSceneSequenceEditor* SequenceEditor = FMovieSceneSequenceEditor::Find(WidgetAnimation);
+			if (!SequenceEditor)
+			{
+				return;
+			}
+
+			UBlueprint* SequenceDirectorBP = SequenceEditor->GetOrCreateDirectorBlueprint(WidgetAnimation);
+			if (!SequenceDirectorBP)
+			{
+				return;
+			}
+
+			FMovieSceneDynamicBindingUtils::EnsureBlueprintExtensionCreated(WidgetAnimation, SequenceDirectorBP);
+			FKismetEditorUtilities::CompileBlueprint(SequenceDirectorBP);
+		}
+	}
+
+
 private:
 	TSharedPtr<FExtensibilityManager> MenuExtensibilityManager;
 	TSharedPtr<FExtensibilityManager> ToolBarExtensibilityManager;
@@ -476,6 +536,8 @@ private:
 	FDelegateHandle BlueprintVariableCustomizationHandle;
 	/** Handle for FBlueprintEditorModule::RegisterFunctionCustomization */
 	FDelegateHandle BlueprintFunctionCustomizationHandle;
+	FDelegateHandle FixupDynamicBindingPayloadParameterNameHandle;
+	FDelegateHandle FixupWidgetDynamicBindingsHandle;
 
 	bool bThumbnailRenderersRegistered;
 	bool bOnPostEngineInitHandled;
diff --git a/Engine/Source/Editor/UMGEditor/Private/WidgetBlueprint.cpp b/Engine/Source/Editor/UMGEditor/Private/WidgetBlueprint.cpp
--- a/Engine/Source/Editor/UMGEditor/Private/WidgetBlueprint.cpp
+++ b/Engine/Source/Editor/UMGEditor/Private/WidgetBlueprint.cpp
@@ -554,19 +554,6 @@ FDelegateRuntimeBinding FDelegateEditorBinding::ToRuntimeBinding(UWidgetBlueprin
 	return Binding;
 }
 
-bool FWidgetAnimation_DEPRECATED::SerializeFromMismatchedTag(struct FPropertyTag const& Tag, FStructuredArchive::FSlot Slot)
-{
-	static FName AnimationDataName("AnimationData");
-	if(Tag.Type == NAME_StructProperty && Tag.Name == AnimationDataName)
-	{
-		FStructuredArchive::FRecord Record = Slot.EnterRecord();
-		Record << SA_VALUE(TEXT("MovieScene"), MovieScene);
-		Record << SA_VALUE(TEXT("AnimationBindings"), AnimationBindings);
-		return true;
-	}
-
-	return false;
-}
 /////////////////////////////////////////////////////
 // UWidgetBlueprint
 
@@ -1022,29 +1009,6 @@ void UWidgetBlueprint::PostLoad()
 		Widget->ConnectEditorData();
 	});
 
-	if( GetLinkerUEVersion() < VER_UE4_FIXUP_WIDGET_ANIMATION_CLASS )
-	{
-		// Fixup widget animations.
-		for( auto& OldAnim : AnimationData_DEPRECATED )
-		{
-			FName AnimName = OldAnim.MovieScene->GetFName();
-
-			// Rename the old movie scene so we can reuse the name
-			OldAnim.MovieScene->Rename( *MakeUniqueObjectName( this, UMovieScene::StaticClass(), "MovieScene").ToString(), nullptr, REN_ForceNoResetLoaders | REN_DontCreateRedirectors | REN_DoNotDirty | REN_NonTransactional);
-
-			UWidgetAnimation* NewAnimation = NewObject<UWidgetAnimation>(this, AnimName, RF_Transactional);
-
-			OldAnim.MovieScene->Rename(*AnimName.ToString(), NewAnimation, REN_ForceNoResetLoaders | REN_DontCreateRedirectors | REN_DoNotDirty | REN_NonTransactional );
-
-			NewAnimation->MovieScene = OldAnim.MovieScene;
-			NewAnimation->AnimationBindings = OldAnim.AnimationBindings;
-			
-			Animations.Add( NewAnimation );
-		}	
-
-		AnimationData_DEPRECATED.Empty();
-	}
-
 	if ( GetLinkerUEVersion() < VER_UE4_RENAME_WIDGET_VISIBILITY )
 	{
 		static const FName Visiblity(TEXT("Visiblity"));
diff --git a/Engine/Source/Editor/UMGEditor/Private/WidgetBlueprintCompiler.cpp b/Engine/Source/Editor/UMGEditor/Private/WidgetBlueprintCompiler.cpp
--- a/Engine/Source/Editor/UMGEditor/Private/WidgetBlueprintCompiler.cpp
+++ b/Engine/Source/Editor/UMGEditor/Private/WidgetBlueprintCompiler.cpp
@@ -1243,12 +1243,12 @@ void FWidgetBlueprintCompilerContext::ValidateWidgetAnimations()
 		for (const FWidgetAnimationBinding& Binding : InAnimation->AnimationBindings)
 		{
 			// Look for the object bindings within the widget
-			UObject* FoundObject = Binding.FindRuntimeObject(*LatestWidgetTree, *UserWidget);
+			UObject* FoundObject = Binding.FindRuntimeObject(*LatestWidgetTree, *UserWidget, InAnimation, nullptr);
 
 			// If any of the FoundObjects is null, we do not play the animation.
 			if (FoundObject == nullptr)
 			{
-				FoundObject = Binding.FindRuntimeObject(*WidgetBP->WidgetTree, *UserWidget);
+				FoundObject = Binding.FindRuntimeObject(*WidgetBP->WidgetTree, *UserWidget, InAnimation, nullptr);
 				if (FoundObject == nullptr)
 				{
 					// Notify the user of the null track in the editor
diff --git a/Engine/Source/Editor/UMGEditor/Private/WidgetBlueprintEditor.cpp b/Engine/Source/Editor/UMGEditor/Private/WidgetBlueprintEditor.cpp
--- a/Engine/Source/Editor/UMGEditor/Private/WidgetBlueprintEditor.cpp
+++ b/Engine/Source/Editor/UMGEditor/Private/WidgetBlueprintEditor.cpp
@@ -89,6 +89,7 @@
 #include "Preferences/UnrealEdOptions.h"
 #include "UnrealEdGlobals.h"
 #include "GraphEditorActions.h"
+#include "MovieSceneDynamicBindingCustomization.h"
 
 #define LOCTEXT_NAMESPACE "UMG"
 
@@ -660,7 +661,8 @@ bool FWidgetBlueprintEditor::IsBindingSelected(const FMovieSceneBinding& InBindi
 	TSharedPtr<ISequencer>& ActiveSequencer = GetSequencer();
 	UMovieSceneSequence* AnimationSequence = ActiveSequencer->GetFocusedMovieSceneSequence();
 	UObject* BindingContext = GetAnimationPlaybackContext();
-	TArray<UObject*, TInlineAllocator<1>> BoundObjects = AnimationSequence->LocateBoundObjects(InBinding.GetObjectGuid(), BindingContext);
+	TArray<UObject*, TInlineAllocator<1>> BoundObjects;
+	AnimationSequence->LocateBoundObjects(InBinding.GetObjectGuid(), UE::UniversalObjectLocator::FResolveParams(BindingContext), ActiveSequencer->GetSharedPlaybackState(), BoundObjects);
 
 	if (BoundObjects.Num() == 0)
 	{
@@ -2295,7 +2297,7 @@ void FWidgetBlueprintEditor::RemoveMissingWidgetsFromTrack(FGuid ObjectId)
 	for (int32 Index = WidgetAnimation->AnimationBindings.Num() - 1; Index >= 0; --Index)
 	{
 		const FWidgetAnimationBinding& Binding = WidgetAnimation->AnimationBindings[Index];
-		if (Binding.AnimationGuid == ObjectId && Binding.FindRuntimeObject(*PreviewRoot->WidgetTree, *PreviewRoot) == nullptr)
+		if (Binding.AnimationGuid == ObjectId && Binding.FindRuntimeObject(*PreviewRoot->WidgetTree, *PreviewRoot, WidgetAnimation, ActiveSequencer->GetSharedPlaybackState()) == nullptr)
 		{
 			WidgetAnimation->AnimationBindings.RemoveAt(Index, EAllowShrinking::No);
 		}
@@ -2369,27 +2371,118 @@ void FWidgetBlueprintEditor::AddDynamicPossessionMenu(FMenuBuilder& MenuBuilder,
 
 	TSharedPtr<ISequencer>& ActiveSequencer = GetSequencer();
 
-	UMovieScene* MovieScene = ActiveSequencer->GetFocusedMovieSceneSequence()->GetMovieScene();
+	UWidgetAnimation* WidgetAnimation = Cast<UWidgetAnimation>(ActiveSequencer->GetFocusedMovieSceneSequence());
+	UMovieScene* MovieScene = WidgetAnimation->GetMovieScene();
+
 	FMovieScenePossessable* Possessable = MovieScene->FindPossessable(ObjectId);
 	if (!Possessable)
 	{
 		return;
 	}
 
-	TSharedPtr<FSequencerEditorViewModel> SequencerViewModel = ActiveSequencer->GetViewModel();
-	FObjectBindingModelStorageExtension* ObjectStorage = SequencerViewModel->GetRootModel()->CastDynamic<FObjectBindingModelStorageExtension>();
-	if (!ObjectStorage)
+	FWidgetAnimationBinding* WidgetBinding = nullptr;
+
+	for (int32 Index = 0; Index < WidgetAnimation->AnimationBindings.Num(); ++Index)
 	{
-		return;
+		FWidgetAnimationBinding& Binding = WidgetAnimation->AnimationBindings[Index];
+		if (Binding.AnimationGuid == ObjectId)
+		{
+			WidgetBinding = &Binding;
+			break;
+		}
 	}
 
-	TSharedPtr<FObjectBindingModel> ObjectBindingModel = ObjectStorage->FindModelForObjectBinding(ObjectId);
-	if (!ObjectBindingModel)
+	if (WidgetBinding)
+	{
+		FDetailsViewArgs DetailsViewArgs;
+		{
+			DetailsViewArgs.bAllowSearch = false;
+			DetailsViewArgs.bCustomFilterAreaLocation = true;
+			DetailsViewArgs.bCustomNameAreaLocation = true;
+			DetailsViewArgs.bHideSelectionTip = true;
+			DetailsViewArgs.bLockable = false;
+			DetailsViewArgs.bSearchInitialKeyFocus = true;
+			DetailsViewArgs.bUpdatesFromSelection = false;
+			DetailsViewArgs.bShowOptions = false;
+			DetailsViewArgs.bShowModifiedPropertiesOption = false;
+			DetailsViewArgs.bShowScrollBar = false;
+		}
+
+		FStructureDetailsViewArgs StructureViewArgs;
+		{
+			StructureViewArgs.bShowObjects = false;
+			StructureViewArgs.bShowAssets = true;
+			StructureViewArgs.bShowClasses = true;
+			StructureViewArgs.bShowInterfaces = false;
+		}
+
+		TSharedRef<IStructureDetailsView> StructureDetailsView = FModuleManager::GetModuleChecked<FPropertyEditorModule>("PropertyEditor")
+			.CreateStructureDetailView(DetailsViewArgs, StructureViewArgs, nullptr);
+
+		// Register details customizations for this instance
+		StructureDetailsView->GetDetailsView()->RegisterInstancedCustomPropertyTypeLayout(
+			FMovieSceneDynamicBinding::StaticStruct()->GetFName(),
+			FOnGetPropertyTypeCustomizationInstance::CreateStatic(&FMovieSceneDynamicBindingCustomization::MakeInstance, MovieScene, ObjectId));
+
+		// We can't just show the FMovieSceneDynamicBinding struct in the details view, because Slate only uses
+		// the above details view customization for *properties* (not for the root object). So here we put a copy of
+		// our dynamic binding struct inside a container, and when the details view is done setting values on it,
+		// we copy these values back to the original dynamic binding.
+		TSharedPtr<FStructOnScope> StructOnScope = MakeShared<FStructOnScope>(FMovieSceneDynamicBindingContainer::StaticStruct());
+		FMovieSceneDynamicBindingContainer* BufferContainer = (FMovieSceneDynamicBindingContainer*)StructOnScope->GetStructMemory();
+		BufferContainer->DynamicBinding = WidgetBinding->DynamicBinding;
+		StructureDetailsView->SetStructureData(StructOnScope);
+
+		StructureDetailsView->GetOnFinishedChangingPropertiesDelegate().AddSP(this, &FWidgetBlueprintEditor::OnFinishedChangingDynamicBindingProperties, StructOnScope, ObjectId);
+
+		MenuBuilder.BeginSection(NAME_None, LOCTEXT("DynamicBindingHeader", "Dynamic Binding"));
+		{
+			TSharedRef<SWidget> Widget = StructureDetailsView->GetWidget().ToSharedRef();
+			MenuBuilder.AddWidget(Widget, FText());
+		}
+		MenuBuilder.EndSection();
+	}
+}
+
+void FWidgetBlueprintEditor::OnFinishedChangingDynamicBindingProperties(const FPropertyChangedEvent& ChangeEvent, TSharedPtr<FStructOnScope> ValueStruct, FGuid ObjectId)
+{
+	auto* Container = (FMovieSceneDynamicBindingContainer*)ValueStruct->GetStructMemory();
+
+	using namespace UE::Sequencer;
+
+	TSharedPtr<ISequencer>& ActiveSequencer = GetSequencer();
+
+	UWidgetAnimation* WidgetAnimation = Cast<UWidgetAnimation>(ActiveSequencer->GetFocusedMovieSceneSequence());
+	UMovieScene* MovieScene = WidgetAnimation->GetMovieScene();
+
+	FMovieScenePossessable* Possessable = MovieScene->FindPossessable(ObjectId);
+	if (!Possessable)
 	{
 		return;
 	}
 
-	ObjectBindingModel->AddDynamicBindingMenu(MenuBuilder, Possessable->DynamicBinding);
+	FWidgetAnimationBinding* WidgetBinding = nullptr;
+
+	for (int32 Index = 0; Index < WidgetAnimation->AnimationBindings.Num(); ++Index)
+	{
+		FWidgetAnimationBinding& Binding = WidgetAnimation->AnimationBindings[Index];
+		if (Binding.AnimationGuid == ObjectId)
+		{
+			WidgetBinding = &Binding;
+			break;
+		}
+	}
+
+	if (WidgetBinding)
+	{
+		WidgetBinding->DynamicBinding = Container->DynamicBinding;
+	}
+
+	// Force refresh the binding
+	if (FMovieSceneEvaluationState* EvaluationState = ActiveSequencer->GetSharedPlaybackState()->FindCapability<FMovieSceneEvaluationState>())
+	{
+		EvaluationState->Invalidate(ObjectId, ActiveSequencer->GetFocusedTemplateID());
+	}
 }
 
 void FWidgetBlueprintEditor::AddSlotTrack( UPanelSlot* Slot )
@@ -2514,7 +2607,8 @@ void FWidgetBlueprintEditor::SyncSelectedWidgetsWithSequencerSelection(TArray<FG
 	TSet<FWidgetReference> SequencerSelectedWidgets;
 	for (FGuid Guid : ObjectGuids)
 	{
-		TArray<UObject*, TInlineAllocator<1>> BoundObjects = AnimationSequence->LocateBoundObjects(Guid, BindingContext);
+		TArray<UObject*, TInlineAllocator<1>> BoundObjects;
+		AnimationSequence->LocateBoundObjects(Guid, UE::UniversalObjectLocator::FResolveParams(BindingContext), ActiveSequencer->GetSharedPlaybackState(), BoundObjects);
 		if (BoundObjects.Num() == 0)
 		{
 			continue;
@@ -2588,7 +2682,7 @@ void FWidgetBlueprintEditor::UpdateTrackName(FGuid ObjectId)
 		}
 
 		TArray<UObject*, TInlineAllocator<1>> BoundObjects;
-		WidgetAnimation->LocateBoundObjects(ObjectId, BindingContext, BoundObjects);
+		WidgetAnimation->LocateBoundObjects(ObjectId, UE::UniversalObjectLocator::FResolveParams(BindingContext), ActiveSequencer->GetSharedPlaybackState(), BoundObjects);
 
 		if (BoundObjects.Num() > 0)
 		{
diff --git a/Engine/Source/Editor/UMGEditor/Public/WidgetBlueprint.h b/Engine/Source/Editor/UMGEditor/Public/WidgetBlueprint.h
--- a/Engine/Source/Editor/UMGEditor/Public/WidgetBlueprint.h
+++ b/Engine/Source/Editor/UMGEditor/Public/WidgetBlueprint.h
@@ -173,32 +173,6 @@ struct UMGEDITOR_API FDelegateEditorBinding
 	FDelegateRuntimeBinding ToRuntimeBinding(class UWidgetBlueprint* Blueprint) const;
 };
 
-
-/** Struct used only for loading old animations */
-USTRUCT()
-struct FWidgetAnimation_DEPRECATED
-{
-	GENERATED_USTRUCT_BODY()
-
-	UPROPERTY()
-	TObjectPtr<UMovieScene> MovieScene = nullptr;
-
-	UPROPERTY()
-	TArray<FWidgetAnimationBinding> AnimationBindings;
-
-	bool SerializeFromMismatchedTag(struct FPropertyTag const& Tag, FStructuredArchive::FSlot Slot);
-
-};
-
-template<>
-struct TStructOpsTypeTraits<FWidgetAnimation_DEPRECATED> : public TStructOpsTypeTraitsBase2<FWidgetAnimation_DEPRECATED>
-{
-	enum
-	{
-		WithStructuredSerializeFromMismatchedTag = true,
-	};
-};
-
 UENUM()
 enum class EWidgetSupportsDynamicCreation : uint8
 {
@@ -252,9 +226,6 @@ public:
 	UPROPERTY()
 	TArray< FDelegateEditorBinding > Bindings;
 
-	UPROPERTY()
-	TArray<FWidgetAnimation_DEPRECATED> AnimationData_DEPRECATED;
-
 	UPROPERTY()
 	TArray<TObjectPtr<UWidgetAnimation>> Animations;
 
diff --git a/Engine/Source/Editor/UMGEditor/Public/WidgetBlueprintEditor.h b/Engine/Source/Editor/UMGEditor/Public/WidgetBlueprintEditor.h
--- a/Engine/Source/Editor/UMGEditor/Public/WidgetBlueprintEditor.h
+++ b/Engine/Source/Editor/UMGEditor/Public/WidgetBlueprintEditor.h
@@ -435,6 +435,9 @@ private:
 	/** Update the name of a track to reflect changes in bindings */
 	void UpdateTrackName(FGuid ObjectId);
 
+	/* Called after changing dynamic binding properties*/
+	void OnFinishedChangingDynamicBindingProperties(const FPropertyChangedEvent& ChangeEvent, TSharedPtr<FStructOnScope> ValueStruct, FGuid ObjectId);
+
 private:
 	/** The preview scene that owns the preview GUI */
 	FPreviewScene PreviewScene;
diff --git a/Engine/Source/Runtime/Core/Public/UObject/FortniteMainBranchObjectVersion.h b/Engine/Source/Runtime/Core/Public/UObject/FortniteMainBranchObjectVersion.h
--- a/Engine/Source/Runtime/Core/Public/UObject/FortniteMainBranchObjectVersion.h
+++ b/Engine/Source/Runtime/Core/Public/UObject/FortniteMainBranchObjectVersion.h
@@ -449,6 +449,8 @@ struct FFortniteMainBranchObjectVersion
 
 		// Add the imported asset dependencies to the Cloth Asset USD Import node
 		ChaosClothAssetUSDImportNodeAddAssetDependencies,
+		// Upgraded movie scene 'dynamic bindings' to use the new Custom Bindings system
+		LevelSequenceUpgradeDynamicBindings,
 
 		// Changed HitLighting to HitLightingForReflections, and HitLighting now means hit lighting for entire Lumen
 		LumenRayLightingModeOverrideEnum,
diff --git a/Engine/Source/Runtime/LevelSequence/Private/LevelSequence.cpp b/Engine/Source/Runtime/LevelSequence/Private/LevelSequence.cpp
--- a/Engine/Source/Runtime/LevelSequence/Private/LevelSequence.cpp
+++ b/Engine/Source/Runtime/LevelSequence/Private/LevelSequence.cpp
@@ -49,6 +49,12 @@
 #include "Engine/AssetUserData.h"
 #include "Misc/App.h"
 #include "Misc/DateTime.h"
+#include "UObject/FortniteMainBranchObjectVersion.h"
+#include "Bindings/MovieSceneSpawnableDirectorBlueprintBinding.h"
+#include "Bindings/MovieSceneReplaceableDirectorBlueprintBinding.h"
+#include "Bindings/MovieSceneSpawnableActorBinding.h"
+#include "MovieSceneFolder.h"
+#include "Sections/MovieSceneBindingLifetimeSection.h"
 
 #include UE_INLINE_GENERATED_CPP_BY_NAME(LevelSequence)
 
@@ -58,6 +64,7 @@
 	#include "UObject/ObjectRedirector.h"
 
 ULevelSequence::FPostDuplicateEvent ULevelSequence::PostDuplicateEvent;
+ULevelSequence::FFixupDynamicBindingsEvent ULevelSequence::FixupDynamicBindingsEvent;
 
 #endif
 
@@ -413,9 +420,184 @@ void ULevelSequence::PostLoad()
 	}
 	ObjectReferences_DEPRECATED.Map.Empty();
 
+	if (GetLinkerCustomVersion(FFortniteMainBranchObjectVersion::GUID) < FFortniteMainBranchObjectVersion::LevelSequenceUpgradeDynamicBindings)
+	{
+		bool bConvertedDynamicBinding = ConvertOldSpawnables();
+
+		for (int32 Index = 0; Index < MovieScene->GetPossessableCount(); ++Index)
+		{
+			FMovieScenePossessable& Possessable = MovieScene->GetPossessable(Index);
+			if (Possessable.DynamicBinding_DEPRECATED.Function)
+			{
+				bConvertedDynamicBinding = true;
+				ConvertDynamicBindingPossessable(Possessable);
+			}
+		}
+		if (bConvertedDynamicBinding && FixupDynamicBindingsEvent.IsBound())
+		{
+			FixupDynamicBindingsEvent.Broadcast(this);
+		}
+	}
+
 #endif
 }
 
+#if WITH_EDITOR
+
+bool UpdateFolderBindingID(UMovieSceneFolder* Folder, FGuid OldGuid, FGuid NewGuid)
+{
+	for (FGuid ChildGuid : Folder->GetChildObjectBindings())
+	{
+		if (ChildGuid == OldGuid)
+		{
+			Folder->AddChildObjectBinding(NewGuid);
+			Folder->RemoveChildObjectBinding(OldGuid);
+			return true;
+		}
+	}
+
+	for (UMovieSceneFolder* ChildFolder : Folder->GetChildFolders())
+	{
+		if (UpdateFolderBindingID(ChildFolder, OldGuid, NewGuid))
+		{
+			return true;
+		}
+	}
+
+	return false;
+}
+
+bool ULevelSequence::ConvertOldSpawnables()
+{
+	bool bConvertedDynamicBinding = false;
+	while (MovieScene->GetSpawnableCount() > 0)
+	{
+		FMovieSceneSpawnable& Spawnable = MovieScene->GetSpawnable(0);
+
+		FMovieScenePossessable* CreatedPossessable = nullptr;
+
+		UObject* ObjectToConvert = Spawnable.GetObjectTemplate();
+
+		UClass* CustomBindingType = nullptr;
+		if (Spawnable.DynamicBinding_DEPRECATED.Function)
+		{
+			CustomBindingType = UMovieSceneSpawnableDirectorBlueprintBinding::StaticClass();
+		}
+		else
+		{
+			// Search through custom binding types to find one that best supports the template type
+			static TArray<const TSubclassOf<UMovieSceneCustomBinding>> CachedCustomBindingTypes;
+			static bool CustomBindingTypesCached = false;
+			if (!CustomBindingTypesCached)
+			{
+				CustomBindingTypesCached = true;
+				MovieSceneHelpers::GetPrioritySortedCustomBindingTypes(CachedCustomBindingTypes);
+			}
+
+			for (const TSubclassOf<UMovieSceneCustomBinding>& CandidateCustomBindingType : CachedCustomBindingTypes)
+			{
+				if (CandidateCustomBindingType && CandidateCustomBindingType->IsChildOf(UMovieSceneSpawnableBindingBase::StaticClass()) && CandidateCustomBindingType->GetDefaultObject<UMovieSceneCustomBinding>()->SupportsBindingCreationFromObject(Spawnable.GetObjectTemplate()))
+				{
+					CustomBindingType = CandidateCustomBindingType;
+					break;
+				}
+			}
+		}
+
+		if (!CustomBindingType)
+		{
+			UE_LOG(LogLevelSequence, Warning, TEXT("Could not upgrade Spawnable '%s' with ID '%s'"), *Spawnable.GetName(), *Spawnable.GetGuid().ToString());
+			break;
+		}
+
+		UMovieSceneCustomBinding* NewCustomBinding = CustomBindingType->GetDefaultObject<UMovieSceneCustomBinding>()->CreateNewCustomBinding(ObjectToConvert, *MovieScene);
+
+		if (!NewCustomBinding)
+		{
+			UE_LOG(LogLevelSequence, Warning, TEXT("Could not upgrade Spawnable '%s' with ID '%s'"), *Spawnable.GetName(), *Spawnable.GetGuid().ToString());
+			break;
+		}
+
+		if (UMovieSceneSpawnableDirectorBlueprintBinding* DirectorBlueprintBinding = Cast<UMovieSceneSpawnableDirectorBlueprintBinding>(NewCustomBinding))
+		{
+			// Copy over the binding info
+			DirectorBlueprintBinding->DynamicBinding = Spawnable.DynamicBinding_DEPRECATED;
+			bConvertedDynamicBinding = true;
+		}
+		else if (UMovieSceneSpawnableActorBinding* SpawnableActorBinding = Cast<UMovieSceneSpawnableActorBinding>(NewCustomBinding))
+		{
+			SpawnableActorBinding->bNetAddressableName = Spawnable.bNetAddressableName;
+			SpawnableActorBinding->LevelName = Spawnable.LevelName;
+			SpawnableActorBinding->bContinuouslyRespawn = Spawnable.bContinuouslyRespawn;
+			SpawnableActorBinding->SpawnOwnership = Spawnable.GetSpawnOwnership();
+		}
+
+		FString PossessableName = Spawnable.GetName();
+		FGuid SpawnableGuid = Spawnable.GetGuid();
+
+		FMovieScenePossessable NewPossessable(PossessableName, NewCustomBinding->GetBoundObjectClass());
+		// Steal guid
+		NewPossessable.SetGuid(SpawnableGuid);
+
+		if (FMovieSceneBinding* SpawnableBinding = MovieScene->FindBinding(SpawnableGuid))
+		{
+			// Copy binding and track references to be tied to the new possessable
+			FMovieSceneBinding PossessableBinding = *SpawnableBinding;
+			// Add the custom binding. We use the spawnable binding here since it won't have a binding reference yet, and we need to steal the id
+			BindingReferences.AddOrReplaceBinding(SpawnableGuid, NewCustomBinding, 0);
+
+			// Remove the spawnable and all its' sub tracks
+			if (MovieScene->RemoveSpawnable(SpawnableGuid))
+			{
+				// Add the new possessable with the copied binding
+				MovieScene->AddPossessable(NewPossessable, PossessableBinding);
+			}
+		}
+	}
+	return bConvertedDynamicBinding;
+}
+
+void ULevelSequence::ConvertDynamicBindingPossessable(FMovieScenePossessable& Possessable)
+{
+	UMovieSceneReplaceableDirectorBlueprintBinding* NewCustomBinding = nullptr;
+	
+	const FName InstancedBindingName = MakeUniqueObjectName(MovieScene, UObject::StaticClass(), *FString(Possessable.GetName() + TEXT("_CustomBinding")));
+	NewCustomBinding = NewObject<UMovieSceneReplaceableDirectorBlueprintBinding>(MovieScene, UMovieSceneReplaceableDirectorBlueprintBinding::StaticClass(), InstancedBindingName, RF_Transactional);
+	
+	if (!NewCustomBinding)
+	{
+		return;
+	}
+
+	if (UMovieSceneReplaceableDirectorBlueprintBinding* DirectorBlueprintBinding = Cast<UMovieSceneReplaceableDirectorBlueprintBinding>(NewCustomBinding))
+	{
+		// Copy over the binding info
+		DirectorBlueprintBinding->DynamicBinding = Possessable.DynamicBinding_DEPRECATED;
+		Possessable.DynamicBinding_DEPRECATED = FMovieSceneDynamicBinding();
+	}
+
+	// Replace the current binding with the new one. We call RemoveBinding first because if there were multiple bindings for this track,
+	// they would have been overridden with the Dynamic Binding anyway, and so we ensure that stays the same by keeping only one binding
+	BindingReferences.RemoveBinding(Possessable.GetGuid());
+	BindingReferences.AddOrReplaceBinding(Possessable.GetGuid(), NewCustomBinding, 0);
+
+	// Add a binding lifetime track if not present
+	UMovieSceneBindingLifetimeTrack* BindingLifetimeTrack = Cast<UMovieSceneBindingLifetimeTrack>(MovieScene->FindTrack(UMovieSceneBindingLifetimeTrack::StaticClass(), Possessable.GetGuid(), NAME_None));
+	if (!BindingLifetimeTrack)
+	{
+		BindingLifetimeTrack = Cast<UMovieSceneBindingLifetimeTrack>(MovieScene->AddTrack(UMovieSceneBindingLifetimeTrack::StaticClass(), Possessable.GetGuid()));
+	}
+
+	if (BindingLifetimeTrack && BindingLifetimeTrack->GetAllSections().IsEmpty())
+	{
+		UMovieSceneBindingLifetimeSection* BindingLifetimeSection = Cast<UMovieSceneBindingLifetimeSection>(BindingLifetimeTrack->CreateNewSection());
+		BindingLifetimeSection->SetRange(TRange<FFrameNumber>::All());
+		BindingLifetimeTrack->AddSection(*BindingLifetimeSection);
+	}
+}
+
+#endif
+
 #if WITH_EDITORONLY_DATA
 void ULevelSequence::DeclareConstructClasses(TArray<FTopLevelAssetPath>& OutConstructClasses, const UClass* SpecificSubclass)
 {
@@ -468,7 +650,7 @@ void ULevelSequence::LocateBoundObjects(const FGuid& ObjectId, UObject* Context,
 		InResolveBindingParams.WorldPartitionResolveData ? InResolveBindingParams.WorldPartitionResolveData->SourceWorldAssetPath : InResolveBindingParams.StreamedLevelAssetPath
 		);
 
-	LocateBoundObjects(ObjectId, ResolveParams, OutObjects);
+	LocateBoundObjects(ObjectId, ResolveParams, nullptr, OutObjects);
 }
 
 FGuid ULevelSequence::FindBindingFromObject(UObject* InObject, UObject* Context) const
@@ -681,12 +863,6 @@ FGuid ULevelSequence::FindOrAddBinding(UObject* InObject)
 		{
 			ChildPossessable->SetParent(ParentGuid, MovieScene);
 		}
-
-		FMovieSceneSpawnable* ParentSpawnable = MovieScene->FindSpawnable(ParentGuid);
-		if (ParentSpawnable)
-		{
-			ParentSpawnable->AddChildPossessable(NewGuid);
-		}
 	}
 
 	BindPossessableObject(NewGuid, *InObject, BindingContext);
@@ -707,36 +883,16 @@ FGuid ULevelSequence::CreateSpawnable(UObject* ObjectToSpawn)
 		return FGuid();
 	}
 
-	TArray<TSharedRef<IMovieSceneObjectSpawner>> ObjectSpawners;
-
-	// In order to create a spawnable, we have to instantiate all the relevant object spawners for level sequences, and try to create a spawnable from each
-	FLevelSequenceModule& LevelSequenceModule = FModuleManager::LoadModuleChecked<FLevelSequenceModule>("LevelSequence");
-	LevelSequenceModule.GenerateObjectSpawners(ObjectSpawners);
+	FGuid NewGuid = MovieSceneHelpers::TryCreateCustomSpawnableBinding(this, ObjectToSpawn);
 
-	// The first object spawner to return a valid result will win
-	for (TSharedRef<IMovieSceneObjectSpawner> Spawner : ObjectSpawners)
+	UMovieSceneSpawnTrack* NewSpawnTrack = MovieScene->AddTrack<UMovieSceneSpawnTrack>(NewGuid);
+	if (NewSpawnTrack)
 	{
-		TValueOrError<FNewSpawnable, FText> Result = Spawner->CreateNewSpawnableType(*ObjectToSpawn, *MovieScene, nullptr);
-		if (Result.IsValid())
-		{
-			FNewSpawnable& NewSpawnable = Result.GetValue();
+		NewSpawnTrack->Modify();
 
-			NewSpawnable.Name = MovieSceneHelpers::MakeUniqueSpawnableName(MovieScene, NewSpawnable.Name);			
-
-			FGuid NewGuid = MovieScene->AddSpawnable(NewSpawnable.Name, *NewSpawnable.ObjectTemplate);
-
-			UMovieSceneSpawnTrack* NewSpawnTrack = MovieScene->AddTrack<UMovieSceneSpawnTrack>(NewGuid);
-			if (NewSpawnTrack)
-			{
-				NewSpawnTrack->Modify();
-
-				NewSpawnTrack->AddSection(*NewSpawnTrack->CreateNewSection());
-			}
-			return NewGuid;
-		}
+		NewSpawnTrack->AddSection(*NewSpawnTrack->CreateNewSection());
 	}
-
-	return FGuid();
+	return NewGuid;
 }
 
 #endif // WITH_EDITOR
@@ -826,3 +982,21 @@ const TArray<UAssetUserData*>* ULevelSequence::GetAssetUserDataArray() const
 	return &ToRawPtrTArrayUnsafe(AssetUserData);
 }
 
+void ULevelSequence::IterateDynamicBindings(const TSharedRef<UE::MovieScene::FSharedPlaybackState> SharedPlaybackState, TFunction<void(const FGuid&, FMovieSceneDynamicBinding&)> InCallback)
+{
+	for (FMovieSceneBindingReference& BindingReference : BindingReferences.GetAllReferences())
+	{
+		if (BindingReference.CustomBinding)
+		{
+			if (UMovieSceneReplaceableDirectorBlueprintBinding* ReplaceableDirectorBlueprintBinding = Cast<UMovieSceneReplaceableDirectorBlueprintBinding>(BindingReference.CustomBinding))
+			{
+				InCallback(BindingReference.ID, ReplaceableDirectorBlueprintBinding->DynamicBinding);
+			}
+
+			if (UMovieSceneSpawnableDirectorBlueprintBinding* SpawnableDirectorBlueprintBinding = Cast<UMovieSceneSpawnableDirectorBlueprintBinding>(BindingReference.CustomBinding->AsSpawnable(SharedPlaybackState)))
+			{
+				InCallback(BindingReference.ID, SpawnableDirectorBlueprintBinding->DynamicBinding);
+			}
+		}
+	}
+}
diff --git a/Engine/Source/Runtime/LevelSequence/Private/LevelSequencePlayer.cpp b/Engine/Source/Runtime/LevelSequence/Private/LevelSequencePlayer.cpp
--- a/Engine/Source/Runtime/LevelSequence/Private/LevelSequencePlayer.cpp
+++ b/Engine/Source/Runtime/LevelSequence/Private/LevelSequencePlayer.cpp
@@ -128,7 +128,7 @@ void ULevelSequencePlayer::ResolveBoundObjects(UE::UniversalObjectLocator::FReso
 	{
 		if (const FMovieSceneBindingReferences* BindingReferences = InSequence.GetBindingReferences())
 		{
-			FMovieSceneBindingResolveParams BindingResolveParams{ &InSequence, InBindingId, SequenceID };
+			FMovieSceneBindingResolveParams BindingResolveParams{ &InSequence, InBindingId, SequenceID, LocatorResolveParams.Context };
 			if (LocatorResolveParams.ParameterBuffer == nullptr)
 			{
 				// Allocate temporary local buffer for this
@@ -152,13 +152,13 @@ void ULevelSequencePlayer::ResolveBoundObjects(UE::UniversalObjectLocator::FReso
 				TInlineResolveParameterBuffer<128> Buffer;
 				LocatorResolveParams.ParameterBuffer = &Buffer;
 				LocatorResolveParams.ParameterBuffer->AddParameter(FActorLocatorFragmentResolveParameter::ParameterType, WeakStreamingWorld.Get(), ContainerID, SourceAssetPath);
-				InSequence.LocateBoundObjects(InBindingId, LocatorResolveParams, OutObjects);
+				InSequence.LocateBoundObjects(InBindingId, LocatorResolveParams, GetSharedPlaybackState(), OutObjects);
 				LocatorResolveParams.ParameterBuffer = nullptr;
 			}
 			else
 			{
 				LocatorResolveParams.ParameterBuffer->AddParameter(FActorLocatorFragmentResolveParameter::ParameterType, WeakStreamingWorld.Get(), ContainerID, SourceAssetPath);
-				InSequence.LocateBoundObjects(InBindingId, LocatorResolveParams, OutObjects);
+				InSequence.LocateBoundObjects(InBindingId, LocatorResolveParams, GetSharedPlaybackState(), OutObjects);
 			}
 		}
 	}
diff --git a/Engine/Source/Runtime/LevelSequence/Public/LevelSequence.h b/Engine/Source/Runtime/LevelSequence/Public/LevelSequence.h
--- a/Engine/Source/Runtime/LevelSequence/Public/LevelSequence.h
+++ b/Engine/Source/Runtime/LevelSequence/Public/LevelSequence.h
@@ -85,6 +85,11 @@ public:
 
 	DECLARE_DELEGATE_RetVal_OneParam(void, FPostDuplicateEvent, ULevelSequence*);
 	static LEVELSEQUENCE_API FPostDuplicateEvent PostDuplicateEvent;
+
+	DECLARE_MULTICAST_DELEGATE_OneParam(FFixupDynamicBindingsEvent, ULevelSequence*);
+	static LEVELSEQUENCE_API FFixupDynamicBindingsEvent FixupDynamicBindingsEvent;
+
+
 #endif
 
 	LEVELSEQUENCE_API virtual void PostDuplicate(bool bDuplicateForPIE) override;
@@ -97,6 +102,8 @@ public:
 	UE_DEPRECATED(5.4, "Use the base class LocateBoundObjects()")
 	LEVELSEQUENCE_API void LocateBoundObjects(const FGuid& ObjectId, UObject* Context, const FLevelSequenceBindingReference::FResolveBindingParams& InResolveBindingParams, TArray<UObject*, TInlineAllocator<1>>& OutObjects) const;
 
+	LEVELSEQUENCE_API virtual void IterateDynamicBindings(const TSharedRef<UE::MovieScene::FSharedPlaybackState> SharedPlaybackState, TFunction<void(const FGuid&, FMovieSceneDynamicBinding&)> InCallback) override;
+
 #if WITH_EDITOR
 
 
@@ -125,6 +132,14 @@ protected:
 	 */
 	LEVELSEQUENCE_API void OnDirectorRecompiled(UBlueprint*);
 
+private:
+#if WITH_EDITOR
+	// Used to convert any old spawnables in this movie scene to use the custom binding system
+	bool ConvertOldSpawnables();
+	// Used to convert possessables using deprecated 'DynamicBinding' to a custom director blueprint binding
+	void ConvertDynamicBindingPossessable(FMovieScenePossessable& Possessable);
+#endif
+
 #endif // WITH_EDITOR
 
 protected:
diff --git a/Engine/Source/Runtime/MovieScene/Private/Bindings/MovieSceneReplaceableBinding.cpp b/Engine/Source/Runtime/MovieScene/Private/Bindings/MovieSceneReplaceableBinding.cpp
--- a/Engine/Source/Runtime/MovieScene/Private/Bindings/MovieSceneReplaceableBinding.cpp
+++ b/Engine/Source/Runtime/MovieScene/Private/Bindings/MovieSceneReplaceableBinding.cpp
@@ -47,7 +47,7 @@ const FSlateBrush* UMovieSceneReplaceableBindingBase::GetBindingTrackCustomIconO
 
 FText UMovieSceneReplaceableBindingBase::GetBindingTrackIconTooltip() const
 {
-	return LOCTEXT("CustomReplaceableTooltip", "This item is dynamically bound at runtime, and spawns a preview object in Editor within Sequencer");
+	return LOCTEXT("CustomReplaceableTooltip", "This item is dynamically bound at runtime, and may spawn a preview object in Editor within Sequencer");
 }
 
 
@@ -156,14 +156,9 @@ UMovieSceneSpawnableBindingBase* UMovieSceneReplaceableBindingBase::CreateInnerS
 
 UMovieSceneCustomBinding* UMovieSceneReplaceableBindingBase::CreateNewCustomBinding(UObject* SourceObject, UMovieScene& OwnerMovieScene)
 {
-	if (!SourceObject)
-	{
-		return nullptr;
-	}
-
 	UMovieSceneReplaceableBindingBase* NewCustomBinding = nullptr;
 
-	const FName TemplateName = MakeUniqueObjectName(&OwnerMovieScene, UObject::StaticClass(), SourceObject->GetFName());
+	const FName TemplateName = MakeUniqueObjectName(&OwnerMovieScene, UObject::StaticClass(), SourceObject ? SourceObject->GetFName() : GetClass()->GetFName());
 	const FName InstancedBindingName = MakeUniqueObjectName(&OwnerMovieScene, UObject::StaticClass(), *FString(TemplateName.ToString() + TEXT("_CustomBinding")));
 
 	NewCustomBinding = NewObject<UMovieSceneReplaceableBindingBase>(&OwnerMovieScene, GetClass(), InstancedBindingName, RF_Transactional);
diff --git a/Engine/Source/Runtime/MovieScene/Private/Bindings/MovieSceneSpawnableBinding.cpp b/Engine/Source/Runtime/MovieScene/Private/Bindings/MovieSceneSpawnableBinding.cpp
--- a/Engine/Source/Runtime/MovieScene/Private/Bindings/MovieSceneSpawnableBinding.cpp
+++ b/Engine/Source/Runtime/MovieScene/Private/Bindings/MovieSceneSpawnableBinding.cpp
@@ -14,10 +14,13 @@
 #include "Internationalization/Internationalization.h"
 #include "Tracks/MovieSceneSpawnTrack.h"
 #include "Sections/MovieSceneBoolSection.h"
+#include "MovieSceneSequence.h"
 
 #define LOCTEXT_NAMESPACE "FPossessableModel"
 
 
+static const FName SequencerActorTag(TEXT("SequencerActor"));
+
 UObject* UMovieSceneSpawnableBindingBase::SpawnObject(const FGuid& BindingId, int32 BindingIndex, UMovieScene& MovieScene, FMovieSceneSequenceIDRef TemplateID, TSharedRef<const UE::MovieScene::FSharedPlaybackState> SharedPlaybackState)
 {
 	UWorld* WorldContext = GetWorldContext(SharedPlaybackState);
@@ -59,6 +62,64 @@ UObject* UMovieSceneSpawnableBindingBase::SpawnObject(const FGuid& BindingId, in
 	}
 #endif
 
+	// If have spawned an actor, do some actor-specific setup
+	if (AActor* SpawnedActor = Cast<AActor>(SpawnedObject))
+	{
+		// Ensure this spawnable is not a preview actor. Preview actors will not have BeginPlay() called on them.
+#if WITH_EDITOR
+		SpawnedActor->bIsEditorPreviewActor = false;
+#endif
+
+		// tag this actor so we know it was spawned by sequencer
+		SpawnedActor->Tags.AddUnique(SequencerActorTag);
+
+#if WITH_EDITOR
+		if (GIsEditor)
+		{
+			// Explicitly set RF_Transactional on spawned actors so we can undo/redo properties on them.
+			// This particular UObject will be marked RF_Transactional by the caller, but we need to set it on the components.
+
+			for (UActorComponent* Component : SpawnedActor->GetComponents())
+			{
+				if (Component)
+				{
+					Component->SetFlags(RF_Transactional);
+				}
+			}
+		}
+#endif
+
+#if WITH_EDITOR
+		// Don't set the actor label in PIE as this requires flushing async loading.
+		if (WorldContext->WorldType == EWorldType::Editor)
+		{
+			FString BindingName = GetDesiredBindingName();
+			FString ActorLabel = !BindingName.IsEmpty() ? BindingName : SpawnName.ToString();
+
+			if (FMovieScenePossessable* Possessable = MovieScene.FindPossessable(BindingId))
+			{
+				ActorLabel = Possessable->GetName();
+
+				if (UMovieSceneSequence* Sequence = MovieScene.GetTypedOuter<UMovieSceneSequence>())
+				{
+					if (const FMovieSceneBindingReferences* BindingReferences = Sequence->GetBindingReferences())
+					{
+						if (BindingReferences->GetReferences(BindingId).Num() > 1)
+						{
+							// If there are multiple bound objects, use the Object Template actor label instead of the possessable name
+							if (AActor* ActorTemplate = Cast<AActor>(GetObjectTemplate()))
+							{
+								ActorLabel = ActorTemplate->GetActorLabel();
+							}
+						}
+					}
+				}
+			}
+			SpawnedActor->SetActorLabel(ActorLabel);
+		}
+#endif
+	}
+
 	// Allows derived classes to perform post-spawn logic such as mesh setup on actors.
 	PostSpawnObject(SpawnedObject, WorldContext, BindingId, BindingIndex, MovieScene, TemplateID, SharedPlaybackState);
 
@@ -77,6 +138,19 @@ void UMovieSceneSpawnableBindingBase::DestroySpawnedObject(UObject* Object)
 		// Explicitly remove RF_Transactional on spawned objects since we don't want to transact spawn/destroy events
 		Object->ClearFlags(RF_Transactional);
 	}
+	if (AActor* Actor = Cast<AActor>(Object))
+	{
+		// Explicitly remove RF_Transactional on spawned actors since we don't want to trasact spawn/destroy events
+		// This particular UObject will have RF_Transactional cleared by the caller, but we need to cleared it on the components.
+		for (UActorComponent* Component : Actor->GetComponents())
+		{
+			if (Component)
+			{
+				Component->ClearFlags(RF_Transactional);
+			}
+		}
+	}
+
 #endif
 
 	DestroySpawnedObjectInternal(Object);
@@ -131,7 +205,7 @@ const FSlateBrush* UMovieSceneSpawnableBindingBase::GetBindingTrackCustomIconOve
 
 FText UMovieSceneSpawnableBindingBase::GetBindingTrackIconTooltip() const
 {
-	return LOCTEXT("CustomSpawnableTooltip", "This item is spawned by sequencer by a custom spawnable binding according to this object's binding lifetime track.");
+	return LOCTEXT("CustomSpawnableTooltip", "This item is spawned by sequencer by a custom spawnable binding according to this object's spawn track.");
 }
 
 #endif
diff --git a/Engine/Source/Runtime/MovieScene/Private/Evaluation/MovieSceneEvaluationState.cpp b/Engine/Source/Runtime/MovieScene/Private/Evaluation/MovieSceneEvaluationState.cpp
--- a/Engine/Source/Runtime/MovieScene/Private/Evaluation/MovieSceneEvaluationState.cpp
+++ b/Engine/Source/Runtime/MovieScene/Private/Evaluation/MovieSceneEvaluationState.cpp
@@ -11,7 +11,6 @@
 #include "IMovieScenePlaybackClient.h"
 #include "IMovieScenePlayer.h"
 #include "MovieScene.h"
-#include "MovieSceneDynamicBindingInvoker.h"
 #include "MovieSceneObjectBindingID.h"
 #include "MovieSceneSequence.h"
 #include "UniversalObjectLocatorResolveParams.h"
@@ -434,28 +433,14 @@ void FMovieSceneObjectCache::UpdateBindings(const FGuid& InGuid, TSharedRef<cons
 					}
 					else
 					{
-						FMovieSceneDynamicBindingResolveResult ResolveResult = FMovieSceneDynamicBindingInvoker::ResolveDynamicBinding(SharedPlaybackState, Sequence, SequenceID, *Possessable);
-						if (ResolveResult.Object)
+						UE::UniversalObjectLocator::FResolveParams ResolveParams(ResolutionContext);
+						if (Player)
 						{
-							if (!ResolveResult.bIsPossessedObject)
-							{
-								UE_LOG(LogMovieScene, Error, 
-									TEXT("Possessable '%s' (dynamically resolved to '%s') can't have spawnable-type ownership. The user-defined director blueprint endpoint should set bIsPossessedObject to true."),
-									*LexToString(Possessable->GetName()), *ResolveResult.Object->GetName());
-							}
-							FoundObjects.Add(ResolveResult.Object);
+							Player->ResolveBoundObjects(ResolveParams, InGuid, SequenceID, *Sequence, FoundObjects);
 						}
 						else
 						{
-							UE::UniversalObjectLocator::FResolveParams ResolveParams(ResolutionContext);
-							if (Player)
-							{
-								Player->ResolveBoundObjects(ResolveParams, InGuid, SequenceID, *Sequence, FoundObjects);
-							}
-							else
-							{
-								Sequence->LocateBoundObjects(InGuid, ResolveParams, FoundObjects);
-							}
+							Sequence->LocateBoundObjects(InGuid, ResolveParams, SharedPlaybackState, FoundObjects);
 						}
 					}
 					
@@ -486,28 +471,14 @@ void FMovieSceneObjectCache::UpdateBindings(const FGuid& InGuid, TSharedRef<cons
 				}
 				else
 				{
-					FMovieSceneDynamicBindingResolveResult ResolveResult = FMovieSceneDynamicBindingInvoker::ResolveDynamicBinding(SharedPlaybackState, Sequence, SequenceID, *Possessable);
-					if (ResolveResult.Object)
+					UE::UniversalObjectLocator::FResolveParams ResolveParams(ResolutionContext);
+					if (Player)
 					{
-						if (!ResolveResult.bIsPossessedObject)
-						{
-							UE_LOG(LogMovieScene, Error,
-								TEXT("Possessable '%s' (dynamically resolved to '%s') can't have spawnable-type ownership. The user-defined director blueprint endpoint should set bIsPossessedObject to true."),
-								*LexToString(Possessable->GetName()), *ResolveResult.Object->GetName());
-						}
-						FoundObjects.Add(ResolveResult.Object);
+						Player->ResolveBoundObjects(ResolveParams, InGuid, SequenceID, *Sequence, FoundObjects);
 					}
 					else
 					{
-						UE::UniversalObjectLocator::FResolveParams ResolveParams(ResolutionContext);
-						if (Player)
-						{
-							Player->ResolveBoundObjects(ResolveParams, InGuid, SequenceID, *Sequence, FoundObjects);
-						}
-						else
-						{
-							Sequence->LocateBoundObjects(InGuid, ResolveParams, FoundObjects);
-						}
+						Sequence->LocateBoundObjects(InGuid, ResolveParams, SharedPlaybackState, FoundObjects);
 					}
 				}
 				
diff --git a/Engine/Source/Runtime/MovieScene/Private/MovieScene.cpp b/Engine/Source/Runtime/MovieScene/Private/MovieScene.cpp
--- a/Engine/Source/Runtime/MovieScene/Private/MovieScene.cpp
+++ b/Engine/Source/Runtime/MovieScene/Private/MovieScene.cpp
@@ -283,7 +283,7 @@ void UMovieScene::AddSpawnable(const FMovieSceneSpawnable& InNewSpawnable, const
 	FMovieSceneBinding NewBinding = InNewBinding;
 	for (auto Track : NewBinding.GetTracks())
 	{
-		Track->Rename(nullptr, this);
+		Track->Rename(nullptr, this, REN_DontCreateRedirectors | REN_ForceNoResetLoaders);
 	}
 	int32 NewBindingIndex = InsertSorted(ObjectBindings, MoveTemp(NewBinding));
 
@@ -363,7 +363,7 @@ void UMovieScene::AddPossessable(const FMovieScenePossessable& InNewPossessable,
 	FMovieSceneBinding NewBinding = InNewBinding;
 	for (auto Track : NewBinding.GetTracks())
 	{
-		Track->Rename(nullptr, this);
+		Track->Rename(nullptr, this, REN_DontCreateRedirectors | REN_ForceNoResetLoaders);
 	}
 	int32 NewBindingIndex = InsertSorted(ObjectBindings, MoveTemp(NewBinding));
 	check(ObjectBindings.IsValidIndex(NewBindingIndex));
@@ -1108,7 +1108,7 @@ bool UMovieScene::AddGivenTrack(UMovieSceneTrack* InTrack, const FGuid& ObjectGu
 	Modify();
 	if (FMovieSceneBinding* Binding = FindBinding(ObjectGuid))
 	{
-		InTrack->Rename(nullptr, this);
+		InTrack->Rename(nullptr, this, REN_DontCreateRedirectors | REN_ForceNoResetLoaders);
 		Binding->AddTrack(*InTrack, this);
 		return true;
 	}
@@ -1219,7 +1219,7 @@ bool UMovieScene::AddGivenTrack(UMovieSceneTrack* InTrack)
 	{
 		Modify();
 		Tracks.Add(InTrack);
-		InTrack->Rename(nullptr, this);
+		InTrack->Rename(nullptr, this, REN_DontCreateRedirectors | REN_ForceNoResetLoaders);
 
 		EventHandlers.Trigger(&UE::MovieScene::ISequenceDataEventHandler::OnTrackAdded, InTrack);
 
@@ -1290,7 +1290,7 @@ void UMovieScene::SetCameraCutTrack(UMovieSceneTrack* InTrack)
 	}
 
 	Modify();
-	InTrack->Rename(nullptr, this);
+	InTrack->Rename(nullptr, this, REN_DontCreateRedirectors | REN_ForceNoResetLoaders);
 	UMovieSceneTrack* OldCameraCutTrack = CameraCutTrack;
 	CameraCutTrack = InTrack;
 
@@ -1560,7 +1560,7 @@ void UMovieScene::ReplaceBinding(const FGuid& BindingToReplaceGuid, const FMovie
 		// We also need to change the track's owners to be the MovieScene.
 		for (auto Track : Binding->GetTracks())
 		{
-			Track->Rename(nullptr, this);
+			Track->Rename(nullptr, this, REN_DontCreateRedirectors | REN_ForceNoResetLoaders);
 		}
 
 		EventHandlers.Trigger(&UE::MovieScene::ISequenceDataEventHandler::OnBindingRemoved, BindingToReplaceGuid);
diff --git a/Engine/Source/Runtime/MovieScene/Private/MovieSceneCommonHelpers.cpp b/Engine/Source/Runtime/MovieScene/Private/MovieSceneCommonHelpers.cpp
--- a/Engine/Source/Runtime/MovieScene/Private/MovieSceneCommonHelpers.cpp
+++ b/Engine/Source/Runtime/MovieScene/Private/MovieSceneCommonHelpers.cpp
@@ -19,6 +19,7 @@
 #include "MovieSceneBindingReferences.h"
 #include "Bindings/MovieSceneSpawnableBinding.h"
 #include "Evaluation/MovieSceneEvaluationState.h"
+#include "UObject/UObjectIterator.h"
 
 bool MovieSceneHelpers::IsSectionKeyable(const UMovieSceneSection* Section)
 {
@@ -712,6 +713,64 @@ bool MovieSceneHelpers::IsBoundToSpawnable(UMovieSceneSequence* Sequence, const
 }
 
 
+FGuid MovieSceneHelpers::TryCreateCustomSpawnableBinding(UMovieSceneSequence* Sequence, UObject* CustomBindingObject)
+{
+	FGuid NewID;
+	if (!Sequence)
+	{
+		return NewID;
+	}
+	FMovieSceneBindingReferences* BindingReferences = Sequence->GetBindingReferences();
+	if (!BindingReferences)
+	{
+		return NewID;
+	}
+	static TArray<const TSubclassOf<UMovieSceneCustomBinding>> CachedCustomBindingTypes;
+	static bool CustomBindingTypesCached = false;
+	if (!CustomBindingTypesCached)
+	{
+		CustomBindingTypesCached = true;
+		MovieSceneHelpers::GetPrioritySortedCustomBindingTypes(CachedCustomBindingTypes);
+	}
+
+	UMovieSceneCustomBinding* NewCustomBinding = nullptr;
+
+	for (const TSubclassOf<UMovieSceneCustomBinding>& CustomBindingType : CachedCustomBindingTypes)
+	{
+		// We only want to use children of UMovieSceneSpawnableBindingBase
+		if ((CustomBindingType->IsChildOf<UMovieSceneSpawnableBindingBase>()))
+		{
+			if (UMovieSceneCustomBinding* CustomBindingCDO = CustomBindingType ? CustomBindingType->GetDefaultObject<UMovieSceneCustomBinding>() : nullptr)
+			{
+				if (CustomBindingObject && CustomBindingCDO->SupportsBindingCreationFromObject(CustomBindingObject))
+				{
+					// Create a custom binding from this Object
+					NewCustomBinding = CustomBindingCDO->CreateNewCustomBinding(CustomBindingObject, *Sequence->GetMovieScene());
+					if (NewCustomBinding)
+					{
+						break;
+					}
+				}
+			}
+		}
+	}
+
+	if (NewCustomBinding)
+	{
+		FString DesiredBindingName = NewCustomBinding->GetDesiredBindingName();
+		FString CurrentName = DesiredBindingName.IsEmpty() ? FName::NameToDisplayString(CustomBindingObject->GetName(), false) : DesiredBindingName;
+		CurrentName = MovieSceneHelpers::MakeUniqueBindingName(Sequence->GetMovieScene(), CurrentName);
+
+		NewID = Sequence->GetMovieScene()->AddPossessable(CurrentName, NewCustomBinding->GetBoundObjectClass());
+
+		// Add the custom binding
+		Sequence->GetBindingReferences()->AddOrReplaceBinding(NewID, NewCustomBinding, 0);
+	}
+
+	return NewID;
+}
+
+
 UObject* MovieSceneHelpers::GetSingleBoundObject(UMovieSceneSequence* Sequence, const FGuid& ObjectId, TSharedRef<const UE::MovieScene::FSharedPlaybackState> SharedPlaybackState, int32 BindingIndex)
 {
 	if (Sequence)
@@ -730,13 +789,30 @@ UObject* MovieSceneHelpers::GetSingleBoundObject(UMovieSceneSequence* Sequence,
 						return BoundObjects[0].Get();
 					}
 				}
-
-				const FMovieSceneBindingReferences* Refs = Sequence->GetBindingReferences();
-				if (Refs)
+				else if (const FMovieScenePossessable* Possessable = MovieScene->FindPossessable(ObjectId))
 				{
-					UE::UniversalObjectLocator::FResolveParams LocatorResolveParams(SharedPlaybackState->GetPlaybackContext());
-					FMovieSceneBindingResolveParams BindingResolveParams{ Sequence, ObjectId, SequenceID };
-					return Refs->ResolveSingleBinding(BindingResolveParams, BindingIndex, LocatorResolveParams, SharedPlaybackState);
+					const FMovieSceneBindingReferences* Refs = Sequence->GetBindingReferences();
+					if (Refs)
+					{
+						UObject* ResolutionContext = MovieSceneHelpers::GetResolutionContext(Sequence, ObjectId, SequenceID, SharedPlaybackState);
+
+						if (Possessable->GetParent().IsValid() && Sequence->AreParentContextsSignificant())
+						{
+							TArrayView<TWeakObjectPtr<>> ParentBoundObjects = EvaluationState->FindBoundObjects(FMovieSceneEvaluationOperand(SequenceID, Possessable->GetParent()), SharedPlaybackState);
+							for (TWeakObjectPtr<> Parent : ParentBoundObjects)
+							{
+								ResolutionContext = Parent.Get();
+								if (!ResolutionContext)
+								{
+									continue;
+								}
+							}
+						}
+
+						UE::UniversalObjectLocator::FResolveParams LocatorResolveParams(ResolutionContext);
+						FMovieSceneBindingResolveParams BindingResolveParams{ Sequence, ObjectId, SequenceID, ResolutionContext };
+						return Refs->ResolveSingleBinding(BindingResolveParams, BindingIndex, LocatorResolveParams, SharedPlaybackState);
+					}
 				}
 			}
 		}
@@ -896,6 +972,64 @@ const UClass* MovieSceneHelpers::GetBoundObjectClass(UMovieSceneSequence* Sequen
 #endif
 
 
+void MovieSceneHelpers::GetPrioritySortedCustomBindingTypes(TArray<const TSubclassOf<UMovieSceneCustomBinding>>& OutCustomBindingTypes)
+{
+	OutCustomBindingTypes.Empty();
+	for (TObjectIterator<UClass> It; It; ++It)
+	{
+		if (It->IsChildOf(UMovieSceneCustomBinding::StaticClass()) && !It->HasAnyClassFlags(CLASS_Abstract))
+		{
+			OutCustomBindingTypes.Add(*It);
+		}
+	}// Sort by spawner priority to allow disambiguation for similar object types
+	OutCustomBindingTypes.Sort([](const TSubclassOf<UMovieSceneCustomBinding>& A, const TSubclassOf<UMovieSceneCustomBinding>& B) {
+		return A && B && A->GetDefaultObject<UMovieSceneCustomBinding>()->GetCustomBindingPriority() > B->GetDefaultObject<UMovieSceneCustomBinding>()->GetCustomBindingPriority(); });
+}
+
+
+TSharedRef<UE::MovieScene::FSharedPlaybackState> MovieSceneHelpers::CreateTransientSharedPlaybackState(UObject* WorldContext, UMovieSceneSequence* Sequence)
+{
+	verify(WorldContext && Sequence);
+
+	using namespace UE::MovieScene;
+	FSharedPlaybackStateCreateParams CreateParams;
+	CreateParams.PlaybackContext = WorldContext;
+	TSharedRef<FSharedPlaybackState> TransientPlaybackState = MakeShared<FSharedPlaybackState>(*Sequence, CreateParams);
+
+	FMovieSceneEvaluationState State;
+	TransientPlaybackState->AddCapabilityRaw(&State);
+	State.AssignSequence(MovieSceneSequenceID::Root, *Sequence, TransientPlaybackState);
+
+	return TransientPlaybackState;
+}
+
+UObject* MovieSceneHelpers::GetResolutionContext(UMovieSceneSequence* Sequence, const FGuid& ObjectId, const FMovieSceneSequenceID& SequenceID, TSharedRef<const UE::MovieScene::FSharedPlaybackState> SharedPlaybackState)
+{
+	if (!Sequence)
+	{
+		return nullptr;
+	}
+	UObject* ResolutionContext = SharedPlaybackState->GetPlaybackContext();
+	if (FMovieScenePossessable* Possessable = Sequence->GetMovieScene()->FindPossessable(ObjectId))
+	{
+		if (Possessable->GetParent().IsValid() && Sequence->AreParentContextsSignificant())
+		{
+			if (FMovieSceneEvaluationState* EvaluationState = SharedPlaybackState->FindCapability<FMovieSceneEvaluationState>())
+			{
+				TArrayView<TWeakObjectPtr<>> ParentBoundObjects = EvaluationState->FindBoundObjects(Possessable->GetParent(), SequenceID, SharedPlaybackState);
+				for (TWeakObjectPtr<> Parent : ParentBoundObjects)
+				{
+					ResolutionContext = Parent.Get();
+					if (!ResolutionContext)
+					{
+						continue;
+					}
+				}
+			}
+		}
+	}
+	return ResolutionContext;
+}
 
 MovieSceneHelpers::FMovieSceneScopedPackageDirtyGuard::FMovieSceneScopedPackageDirtyGuard(USceneComponent* InComponent)
 {
diff --git a/Engine/Source/Runtime/MovieScene/Private/MovieSceneDynamicBindingInvoker.cpp b/Engine/Source/Runtime/MovieScene/Private/MovieSceneDynamicBindingInvoker.cpp
--- a/Engine/Source/Runtime/MovieScene/Private/MovieSceneDynamicBindingInvoker.cpp
+++ b/Engine/Source/Runtime/MovieScene/Private/MovieSceneDynamicBindingInvoker.cpp
@@ -12,44 +12,7 @@
 #include "MovieSceneSequence.h"
 #include "MovieSceneSpawnable.h"
 #include "UObject/UnrealType.h"
-
-FMovieSceneDynamicBindingResolveResult FMovieSceneDynamicBindingInvoker::ResolveDynamicBinding(TSharedRef<const FSharedPlaybackState> SharedPlaybackState, UMovieSceneSequence* Sequence, const FMovieSceneSequenceID& SequenceID, const FMovieScenePossessable& Possessable)
-{
-	FMovieSceneDynamicBindingResolveResult Result = ResolveDynamicBinding(SharedPlaybackState, Sequence, SequenceID, Possessable.GetGuid(), Possessable.DynamicBinding);
-	// Top-level possessables must be actors. Children possessables must be components.
-	if (Result.Object && !Possessable.GetParent().IsValid() && !Result.Object.IsA<AActor>())
-	{
-		UE_LOG(
-			LogMovieScene, Error, 
-			TEXT("Possessable '%s' ('%s') must be bound to an actor, but its dynamic binding resolved it to '%s' which is a '%s'. Aborting dynamic binding and falling back to defaults."),
-			*Possessable.GetName(), *LexToString(Possessable.GetGuid()), *Result.Object->GetName(), *Result.Object->GetClass()->GetName());
-		return FMovieSceneDynamicBindingResolveResult();
-	}
-	else if (Result.Object && Possessable.GetParent().IsValid() && !Result.Object.IsA<UActorComponent>())
-	{
-		UE_LOG(
-			LogMovieScene, Error, 
-			TEXT("Possessable '%s' ('%s') must be bound to a component, but its dynamic binding resolved it to '%s' which is a '%s'. Aborting dynamic binding and falling back to defaults."),
-			*Possessable.GetName(), *LexToString(Possessable.GetGuid()), *Result.Object->GetName(), *Result.Object->GetClass()->GetName());
-		return FMovieSceneDynamicBindingResolveResult();
-	}
-	return Result;
-}
-
-FMovieSceneDynamicBindingResolveResult FMovieSceneDynamicBindingInvoker::ResolveDynamicBinding(TSharedRef<const FSharedPlaybackState> SharedPlaybackState, UMovieSceneSequence* Sequence, const FMovieSceneSequenceID& SequenceID, const FMovieSceneSpawnable& Spawnable)
-{
-	FMovieSceneDynamicBindingResolveResult Result = ResolveDynamicBinding(SharedPlaybackState, Sequence, SequenceID, Spawnable.GetGuid(), Spawnable.DynamicBinding);
-	// Spawnables must be actors.
-	if (Result.Object && !Result.Object.IsA<AActor>())
-	{
-		UE_LOG(
-			LogMovieScene, Error, 
-			TEXT("Spawnable '%s' ('%s') must be bound to an actor, but its dynamic binding resolved it to '%s' which is a '%s'. Aborting dynamic binding and falling back to defaults."),
-			*Spawnable.GetName(), *LexToString(Spawnable.GetGuid()), *Result.Object->GetName(), *Result.Object->GetClass()->GetName());
-		return FMovieSceneDynamicBindingResolveResult();
-	}
-	return Result;
-}
+#include "MovieScene.h"
 
 FMovieSceneDynamicBindingResolveResult FMovieSceneDynamicBindingInvoker::ResolveDynamicBinding(TSharedRef<const FSharedPlaybackState> SharedPlaybackState, UMovieSceneSequence* Sequence, const FMovieSceneSequenceID& SequenceID, const FGuid& InGuid, const FMovieSceneDynamicBinding& DynamicBinding)
 {
@@ -81,9 +44,9 @@ FMovieSceneDynamicBindingResolveResult FMovieSceneDynamicBindingInvoker::Resolve
 	if (!DirectorInstance)
 	{
 #if !NO_LOGGING
-		UE_LOG(LogMovieScene, Warning, 
-				TEXT("%s: Failed to resolve dynamic binding '%s' because no director instance was available."), 
-				*Sequence->GetName(), *DynamicBindingFunc->GetName());
+		UE_LOG(LogMovieScene, Warning,
+			TEXT("%s: Failed to resolve dynamic binding '%s' because no director instance was available."),
+			*Sequence->GetName(), *DynamicBindingFunc->GetName());
 #endif
 		// Fallback to default behavior.
 		return FMovieSceneDynamicBindingResolveResult();
@@ -98,24 +61,22 @@ FMovieSceneDynamicBindingResolveResult FMovieSceneDynamicBindingInvoker::Resolve
 	if (!bIsGameWorld && !DynamicBindingFunc->HasMetaData(NAME_CallInEditor))
 	{
 		UE_LOG(LogMovieScene, Verbose,
-				TEXT("%s: Refusing to resolve dynamic binding '%s' in editor world because function '%s' has 'Call in Editor' set to false."),
-				*Sequence->GetName(), *LexToString(InGuid), *DynamicBindingFunc->GetName());
+			TEXT("%s: Refusing to resolve dynamic binding '%s' in editor world because function '%s' has 'Call in Editor' set to false."),
+			*Sequence->GetName(), *LexToString(InGuid), *DynamicBindingFunc->GetName());
 		// Fallback to default behavior.
 		return FMovieSceneDynamicBindingResolveResult();
 	}
 #endif // WITH_EDITOR
 
 	UE_LOG(LogMovieScene, VeryVerbose,
-			TEXT("%s: Resolving dynamic binding '%s' with function '%s'."),
-			*Sequence->GetName(), *LexToString(InGuid), *DynamicBindingFunc->GetName());
+		TEXT("%s: Resolving dynamic binding '%s' with function '%s'."),
+		*Sequence->GetName(), *LexToString(InGuid), *DynamicBindingFunc->GetName());
 
 	FMovieSceneDynamicBindingResolveParams ResolveParams;
 	ResolveParams.ObjectBindingID = InGuid;
 	ResolveParams.Sequence = Sequence;
 	ResolveParams.RootSequence = SharedPlaybackState->GetRootSequence();
-	FMovieSceneDynamicBindingResolveResult Result = InvokeDynamicBinding(DirectorInstance, DynamicBinding, ResolveParams);
-
-	return Result;
+	return InvokeDynamicBinding(DirectorInstance, DynamicBinding, ResolveParams);
 }
 
 FMovieSceneDynamicBindingResolveResult FMovieSceneDynamicBindingInvoker::InvokeDynamicBinding(UObject* DirectorInstance, const FMovieSceneDynamicBinding& DynamicBinding, const FMovieSceneDynamicBindingResolveParams& ResolveParams)
@@ -125,7 +86,7 @@ FMovieSceneDynamicBindingResolveResult FMovieSceneDynamicBindingInvoker::InvokeD
 	// Do some basic checks.
 	UFunction* DynamicBindingFunc = DynamicBinding.Function.Get();
 	if (!ensure(DynamicBindingFunc))
-	{ 
+	{
 		return Result;
 	}
 
@@ -149,11 +110,11 @@ FMovieSceneDynamicBindingResolveResult FMovieSceneDynamicBindingInvoker::InvokeD
 		if (LocalProp->HasAnyPropertyFlags(CPF_ReturnParm))
 		{
 			ensureMsgf(ReturnProp == nullptr,
-					TEXT("Found more than one return parameter in dynamic binding resolver function!"));
+				TEXT("Found more than one return parameter in dynamic binding resolver function!"));
 			ReturnProp = CastFieldChecked<FStructProperty>(LocalProp);
 		}
 	}
-	
+
 	// Set the resolve parameter struct if we need to pass it to the function.
 	if (FProperty* ResolveParamsProp = DynamicBinding.ResolveParamsProperty.Get())
 	{
diff --git a/Engine/Source/Runtime/MovieScene/Private/MovieScenePlayer.cpp b/Engine/Source/Runtime/MovieScene/Private/MovieScenePlayer.cpp
--- a/Engine/Source/Runtime/MovieScene/Private/MovieScenePlayer.cpp
+++ b/Engine/Source/Runtime/MovieScene/Private/MovieScenePlayer.cpp
@@ -138,12 +138,12 @@ void IMovieScenePlayer::ResolveBoundObjects(UE::UniversalObjectLocator::FResolve
 	{
 		if (const FMovieSceneBindingReferences* BindingReferences = InSequence.GetBindingReferences())
 		{
-			FMovieSceneBindingResolveParams BindingResolveParams{ &InSequence, InBindingId, SequenceID };
-			BindingReferences->ResolveBinding(BindingResolveParams, LocatorResolveParams, ConstCastSharedPtr<const UE::MovieScene::FSharedPlaybackState>(const_cast<IMovieScenePlayer*>(this)->FindSharedPlaybackState()), OutObjects);
+			FMovieSceneBindingResolveParams BindingResolveParams{ &InSequence, InBindingId, SequenceID, LocatorResolveParams.Context };
+			BindingReferences->ResolveBinding(BindingResolveParams, LocatorResolveParams, FindSharedPlaybackState(), OutObjects);
 		}
 		else
 		{
-			InSequence.LocateBoundObjects(InBindingId, LocatorResolveParams, OutObjects);
+			InSequence.LocateBoundObjects(InBindingId, LocatorResolveParams, FindSharedPlaybackState(), OutObjects);
 		}
 	}
 }
@@ -176,12 +176,23 @@ TSharedPtr<UE::MovieScene::FSharedPlaybackState> IMovieScenePlayer::FindSharedPl
 	return GetEvaluationTemplate().GetSharedPlaybackState();
 }
 
+TSharedPtr<const UE::MovieScene::FSharedPlaybackState> IMovieScenePlayer::FindSharedPlaybackState() const
+{
+	return ConstCastSharedPtr<const UE::MovieScene::FSharedPlaybackState>(const_cast<IMovieScenePlayer*>(this)->GetEvaluationTemplate().GetSharedPlaybackState());
+}
+
 TSharedRef<UE::MovieScene::FSharedPlaybackState> IMovieScenePlayer::GetSharedPlaybackState()
 {
 	// ToSharedRef will assert if evaluation template isn't initialized
 	return GetEvaluationTemplate().GetSharedPlaybackState().ToSharedRef();
 }
 
+TSharedRef<const UE::MovieScene::FSharedPlaybackState> IMovieScenePlayer::GetSharedPlaybackState() const
+{
+	// ToSharedRef will assert if evaluation template isn't initialized
+	return ConstCastSharedRef<const UE::MovieScene::FSharedPlaybackState>(const_cast<IMovieScenePlayer*>(this)->GetEvaluationTemplate().GetSharedPlaybackState().ToSharedRef());
+}
+
 void IMovieScenePlayer::ResetDirectorInstances()
 {
 	using namespace UE::MovieScene;
diff --git a/Engine/Source/Runtime/MovieScene/Private/MovieScenePossessable.cpp b/Engine/Source/Runtime/MovieScene/Private/MovieScenePossessable.cpp
--- a/Engine/Source/Runtime/MovieScene/Private/MovieScenePossessable.cpp
+++ b/Engine/Source/Runtime/MovieScene/Private/MovieScenePossessable.cpp
@@ -47,7 +47,7 @@ void FMovieScenePossessable::SetParent(const FGuid& InParentGuid, UMovieScene* O
 void FMovieScenePossessable::FixupPossessedObjectClass(UMovieSceneSequence* InSequence, UObject* Context)
 {
 	TArray<UObject*, TInlineAllocator<1>> BoundObjects;
-	InSequence->LocateBoundObjects(Guid, UE::UniversalObjectLocator::FResolveParams(Context), BoundObjects);
+	InSequence->LocateBoundObjects(Guid, UE::UniversalObjectLocator::FResolveParams(Context), nullptr, BoundObjects);
 
 	TArray<UClass*> Classes;
 	for (UObject* BoundObject : BoundObjects)
diff --git a/Engine/Source/Runtime/MovieScene/Private/MovieSceneSequence.cpp b/Engine/Source/Runtime/MovieScene/Private/MovieSceneSequence.cpp
--- a/Engine/Source/Runtime/MovieScene/Private/MovieSceneSequence.cpp
+++ b/Engine/Source/Runtime/MovieScene/Private/MovieSceneSequence.cpp
@@ -20,6 +20,7 @@
 #include "Compilation/MovieSceneCompiledDataManager.h"
 #include "UniversalObjectLocator.h"
 #include "Bindings/MovieSceneSpawnableBinding.h"
+#include "MovieSceneCommonHelpers.h"
 
 #include UE_INLINE_GENERATED_CPP_BY_NAME(MovieSceneSequence)
 
@@ -67,6 +68,11 @@ FMovieSceneBindingReferences* UMovieSceneSequence::GetBindingReferences()
 }
 
 void UMovieSceneSequence::LocateBoundObjects(const FGuid& ObjectId, const UE::UniversalObjectLocator::FResolveParams& ResolveParams, TArray<UObject*, TInlineAllocator<1>>& OutObjects) const
+{
+	LocateBoundObjects(ObjectId, ResolveParams, nullptr, OutObjects);
+}
+
+void UMovieSceneSequence::LocateBoundObjects(const FGuid& ObjectId, const UE::UniversalObjectLocator::FResolveParams& ResolveParams, TSharedPtr<const FSharedPlaybackState> SharedPlaybackState, TArray<UObject*, TInlineAllocator<1>>& OutObjects) const
 {
 	const FMovieSceneBindingReferences* Refs = GetBindingReferences();
 	if (Refs)
@@ -340,18 +346,14 @@ UMovieSceneCompiledData* UMovieSceneSequence::GetOrCreateCompiledData()
 FGuid UMovieSceneSequence::FindPossessableObjectId(UObject& Object, UObject* Context) const
 {
 	using namespace UE::MovieScene;
-
-	FSharedPlaybackStateCreateParams CreateParams;
-	CreateParams.PlaybackContext = Context;
 	UMovieSceneSequence* ThisSequence = const_cast<UMovieSceneSequence*>(this);
-	TSharedRef<FSharedPlaybackState> TransientPlaybackState = MakeShared<FSharedPlaybackState>(*ThisSequence, CreateParams);
-
-	FMovieSceneEvaluationState State;
-	TransientPlaybackState->AddCapabilityRaw(&State);
-	State.AssignSequence(MovieSceneSequenceID::Root, *ThisSequence, TransientPlaybackState);
-
-	FGuid ExistingID = State.FindObjectId(Object, MovieSceneSequenceID::Root, TransientPlaybackState);
-	return ExistingID;
+	TSharedRef<UE::MovieScene::FSharedPlaybackState> TransientPlaybackState = MovieSceneHelpers::CreateTransientSharedPlaybackState(Context, ThisSequence);
+	if (FMovieSceneEvaluationState* EvaluationState = TransientPlaybackState->FindCapability<FMovieSceneEvaluationState>())
+	{
+		FGuid ExistingID = EvaluationState->FindObjectId(Object, MovieSceneSequenceID::Root, TransientPlaybackState);
+		return ExistingID;
+	}
+	return FGuid();
 }
 
 FMovieSceneObjectBindingID UMovieSceneSequence::FindBindingByTag(FName InBindingName) const
diff --git a/Engine/Source/Runtime/MovieScene/Private/MovieSceneSequencePlayer.cpp b/Engine/Source/Runtime/MovieScene/Private/MovieSceneSequencePlayer.cpp
--- a/Engine/Source/Runtime/MovieScene/Private/MovieSceneSequencePlayer.cpp
+++ b/Engine/Source/Runtime/MovieScene/Private/MovieSceneSequencePlayer.cpp
@@ -249,7 +249,7 @@ void UMovieSceneSequencePlayer::ResolveBoundObjects(UE::UniversalObjectLocator::
 
 	if (bAllowDefault)
 	{
-		InSequence.LocateBoundObjects(InBindingId, ResolveParams, OutObjects);
+		InSequence.LocateBoundObjects(InBindingId, ResolveParams, FindSharedPlaybackState(), OutObjects);
 	}
 }
 
diff --git a/Engine/Source/Runtime/MovieScene/Private/MovieSceneSpawnRegister.cpp b/Engine/Source/Runtime/MovieScene/Private/MovieSceneSpawnRegister.cpp
--- a/Engine/Source/Runtime/MovieScene/Private/MovieSceneSpawnRegister.cpp
+++ b/Engine/Source/Runtime/MovieScene/Private/MovieSceneSpawnRegister.cpp
@@ -4,7 +4,6 @@
 #include "MovieScene.h"
 #include "IMovieScenePlayer.h"
 #include "Evaluation/MovieSceneEvaluationTemplateInstance.h"
-#include "MovieSceneDynamicBindingInvoker.h"
 #include "MovieSceneSpawnableAnnotation.h"
 #include "MovieSceneBindingReferences.h"
 #include "Bindings/MovieSceneSpawnableBinding.h"
@@ -52,22 +51,8 @@ UObject* FMovieSceneSpawnRegister::SpawnObject(const FGuid& BindingId, UMovieSce
 		SpawnOwnership = Spawnable->GetSpawnOwnership();
 		IMovieScenePlayer* Player = UE::MovieScene::FPlayerIndexPlaybackCapability::GetPlayer(SharedPlaybackState);
 
-		// See if there is some dynamic binding logic to invoke, otherwise spawn the actor
-			// See if there is some dynamic binding logic to invoke, otherwise spawn the actor
-	    FMovieSceneDynamicBindingResolveResult ResolveResult = FMovieSceneDynamicBindingInvoker::ResolveDynamicBinding(SharedPlaybackState, Sequence, TemplateID, *Spawnable);
-	    if (ResolveResult.Object)
-	    {
-		    SpawnedActor = ResolveResult.Object;
-		    if (ResolveResult.bIsPossessedObject)
-            {
-                SpawnOwnership = ESpawnOwnership::External;
-            }
-        }
-		if (!SpawnedActor)
-		{
-	        // Call through to the list of spawners to see who can spawn something from this FMovieSceneSpawnable
-			SpawnedActor = SpawnObject(*Spawnable, TemplateID, SharedPlaybackState);
-		}
+		// Call through to the list of spawners to see who can spawn something from this FMovieSceneSpawnable
+		SpawnedActor = SpawnObject(*Spawnable, TemplateID, SharedPlaybackState);
 	}
 	else if (UMovieSceneSequence* MovieSceneSequence = MovieScene.GetTypedOuter<UMovieSceneSequence>())
 	{
diff --git a/Engine/Source/Runtime/MovieScene/Public/Bindings/MovieSceneCustomBinding.h b/Engine/Source/Runtime/MovieScene/Public/Bindings/MovieSceneCustomBinding.h
--- a/Engine/Source/Runtime/MovieScene/Public/Bindings/MovieSceneCustomBinding.h
+++ b/Engine/Source/Runtime/MovieScene/Public/Bindings/MovieSceneCustomBinding.h
@@ -88,6 +88,13 @@ public:
 	*/
 	UMovieSceneSpawnableBindingBase* AsSpawnable(TSharedRef<const UE::MovieScene::FSharedPlaybackState> SharedPlaybackState) { return const_cast<UMovieSceneSpawnableBindingBase*>(const_cast<const UMovieSceneCustomBinding*>(this)->AsSpawnable(SharedPlaybackState)); }
 
+	/*
+	* Must be implemented.
+	* Should return the most specific relevant class of the bound object. Used to populate the FMovieScenePossessable bound object class.
+	*/
+	virtual UClass* GetBoundObjectClass() const PURE_VIRTUAL(UMovieSceneCustomBinding::GetBoundObjectClass, return UObject::StaticClass(););
+
+
 #if WITH_EDITOR
 
 	/*
@@ -120,18 +127,17 @@ public:
 	*/
 	virtual FText GetBindingTypePrettyName() const PURE_VIRTUAL(UMovieSceneCustomBinding::GetBindingTypePrettyName, return FText(););
 
-	/*
-	* Must be implemented.
-	* Should return the most specific relevant class of the bound object. Used to populate the FMovieScenePossessable bound object class.
-	*/
-	virtual UClass* GetBoundObjectClass() const PURE_VIRTUAL(UMovieSceneCustomBinding::GetBoundObjectClass, return UObject::StaticClass(););
-
 
 	/*
 	* Called by UI code to see if this custom binding supports converting to a possessable.
 	*/
 	virtual bool CanConvertToPossessable(const FGuid& Guid, FMovieSceneSequenceIDRef TemplateID, TSharedRef<const UE::MovieScene::FSharedPlaybackState> SharedPlaybackState) const { return true; }
 
+	/*
+	* Called by UI code when the binding has recently been added or modified in the case anything needs to be initialized or modified based on this. 
+	*/
+	virtual void OnBindingAddedOrChanged(UMovieScene& OwnerMovieScene) {}
+
 #endif
 
 };
diff --git a/Engine/Source/Runtime/MovieScene/Public/Bindings/MovieSceneReplaceableBinding.h b/Engine/Source/Runtime/MovieScene/Public/Bindings/MovieSceneReplaceableBinding.h
--- a/Engine/Source/Runtime/MovieScene/Public/Bindings/MovieSceneReplaceableBinding.h
+++ b/Engine/Source/Runtime/MovieScene/Public/Bindings/MovieSceneReplaceableBinding.h
@@ -39,16 +39,30 @@ public:
 
 	/* UMovieSceneCustomBinding overrides */
 	virtual bool SupportsBindingCreationFromObject(const UObject* SourceObject) const override;
+	virtual UClass* GetBoundObjectClass() const override;
 #if WITH_EDITOR
 	virtual	void SetupDefaults(UObject* SpawnedObject, FGuid ObjectBindingId, UMovieScene& OwnerMovieScene, TSharedRef<const UE::MovieScene::FSharedPlaybackState> SharedPlaybackState) override;
 	virtual const FSlateBrush* GetBindingTrackCustomIconOverlay() const override;
 	virtual FText GetBindingTrackIconTooltip() const override;
-	virtual UClass* GetBoundObjectClass() const override;
 
 	virtual bool SupportsConversionFromBinding(const FMovieSceneBindingReference& BindingReference, const UObject* SourceObject) const override;
 	virtual UMovieSceneCustomBinding* CreateCustomBindingFromBinding(const FMovieSceneBindingReference& BindingReference, UObject* SourceObject, UMovieScene& OwnerMovieScene) override;
 #endif
 
+#if WITH_EDITORONLY_DATA
+	// Optional Editor-only preview object
+	UPROPERTY(Instanced, VisibleAnywhere, Category="Editor")
+	TObjectPtr<UMovieSceneSpawnableBindingBase> PreviewSpawnable = nullptr;
+#endif
+
+public:
+
+	/*
+	*  Note that we choose to implement CreateCustomBinding here rather than in subclasses.
+	*  Instead we rely on subclasses to implement CreateInnerSpawnable and InitReplaceableBinding which we call here.
+	*/
+	UMovieSceneCustomBinding* CreateNewCustomBinding(UObject* SourceObject, UMovieScene& OwnerMovieScene) override final;
+
 protected:
 
 	/*
@@ -68,12 +82,6 @@ protected:
 	/* Must be implemented. Called from CreateNewCustomBinding to allow the replaceable to initialize any data members from the source object. */
 	virtual void InitReplaceableBinding(UObject* SourceObject, UMovieScene & OwnerMovieScene) PURE_VIRTUAL(UMovieSceneReplaceableBindingBase::InitReplaceableBinding, return;);
 
-#if WITH_EDITORONLY_DATA
-	// Editor-only preview object
-	UPROPERTY(Instanced, VisibleAnywhere, Category="Editor")
-	TObjectPtr<UMovieSceneSpawnableBindingBase> PreviewSpawnable = nullptr;
-#endif
-
 protected:
 
 	/* UMovieSceneCustomBinding overrides*/
@@ -81,9 +89,4 @@ protected:
 	FMovieSceneBindingResolveResult ResolveBinding(const FMovieSceneBindingResolveParams& ResolveParams, int32 BindingIndex, TSharedRef<const UE::MovieScene::FSharedPlaybackState> SharedPlaybackState) const override final;
 	const UMovieSceneSpawnableBindingBase* AsSpawnable(TSharedRef<const UE::MovieScene::FSharedPlaybackState> SharedPlaybackState) const override final;
 
-	/* 
-	*  Note that we choose to implement CreateCustomBinding here rather than in subclasses.
-	*  Instead we rely on subclasses to implement CreateInnerSpawnable and InitReplaceableBinding which we call here.
-	*/
-	UMovieSceneCustomBinding* CreateNewCustomBinding(UObject* SourceObject, UMovieScene& OwnerMovieScene) override final;
 };
\ No newline at end of file
diff --git a/Engine/Source/Runtime/MovieScene/Public/Evaluation/MovieSceneEvaluationState.h b/Engine/Source/Runtime/MovieScene/Public/Evaluation/MovieSceneEvaluationState.h
--- a/Engine/Source/Runtime/MovieScene/Public/Evaluation/MovieSceneEvaluationState.h
+++ b/Engine/Source/Runtime/MovieScene/Public/Evaluation/MovieSceneEvaluationState.h
@@ -26,8 +26,6 @@
 class IMovieScenePlayer;
 class UMovieSceneSequence;
 class UObject;
-struct FMovieSceneDynamicBinding;
-struct FMovieSceneDynamicBindingResolveParams;
 struct FMovieSceneEvaluationKey;
 struct FMovieSceneObjectBindingID;
 struct FSharedPersistentDataKey;
@@ -215,16 +213,6 @@ private:
 	 */
 	void UpdateBindings(const FGuid& InGuid, TSharedRef<const FSharedPlaybackState> SharedPlaybackState);
 
-	/**
-	 * Handles optional dynamic binding for a given object binding.
-	 */
-	bool ResolveDynamicBinding(const FGuid& InGuid, const FMovieSceneDynamicBinding& DynamicBinding, TSharedRef<const FSharedPlaybackState> SharedPlaybackState, TArray<UObject*, TInlineAllocator<1>>& OutObjects);
-
-	/**
-	 * Invokes the custom function used to resolve a given dynamic object binding.
-	 */
-	UObject* InvokeDynamicBinding(UObject* DirectorInstance, const FMovieSceneDynamicBinding& DynamicBinding, const FMovieSceneDynamicBindingResolveParams& ResolveParams);
-
 	/**
 	 * Invalidate the object bindings for a specific object binding ID
 	 */
diff --git a/Engine/Source/Runtime/MovieScene/Public/IMovieScenePlayer.h b/Engine/Source/Runtime/MovieScene/Public/IMovieScenePlayer.h
--- a/Engine/Source/Runtime/MovieScene/Public/IMovieScenePlayer.h
+++ b/Engine/Source/Runtime/MovieScene/Public/IMovieScenePlayer.h
@@ -397,10 +397,12 @@ public:
 	 * Returns the evaluated sequence instance's shared playback state, if any.
 	 */
 	MOVIESCENE_API TSharedPtr<UE::MovieScene::FSharedPlaybackState> FindSharedPlaybackState();
+	MOVIESCENE_API TSharedPtr<const UE::MovieScene::FSharedPlaybackState> FindSharedPlaybackState() const;
 	/**
 	 * Returns the evaluated sequence instance's shared playback state, asserts if there is none.
 	 */
 	MOVIESCENE_API TSharedRef<UE::MovieScene::FSharedPlaybackState> GetSharedPlaybackState();
+	MOVIESCENE_API TSharedRef<const UE::MovieScene::FSharedPlaybackState> GetSharedPlaybackState() const;
 
 	uint16 GetUniqueIndex() const
 	{
diff --git a/Engine/Source/Runtime/MovieScene/Public/MovieSceneBindingReferences.h b/Engine/Source/Runtime/MovieScene/Public/MovieSceneBindingReferences.h
--- a/Engine/Source/Runtime/MovieScene/Public/MovieSceneBindingReferences.h
+++ b/Engine/Source/Runtime/MovieScene/Public/MovieSceneBindingReferences.h
@@ -58,6 +58,10 @@ struct FMovieSceneBindingResolveParams
 	/** The sequence ID of the object binding being resolved */
 	UPROPERTY()
 	FMovieSceneSequenceID SequenceID;
+
+	/* The outer context with which to resolve this binding. May be the world, or may be an outer UObject.*/
+	UPROPERTY()
+	TObjectPtr<UObject> Context = nullptr;
 };
 
 /**
diff --git a/Engine/Source/Runtime/MovieScene/Public/MovieSceneCommonHelpers.h b/Engine/Source/Runtime/MovieScene/Public/MovieSceneCommonHelpers.h
--- a/Engine/Source/Runtime/MovieScene/Public/MovieSceneCommonHelpers.h
+++ b/Engine/Source/Runtime/MovieScene/Public/MovieSceneCommonHelpers.h
@@ -18,8 +18,11 @@ class UMovieSceneSubSection;
 class UMovieSceneSequence;
 class USceneComponent;
 class USoundBase;
+template <class TClass> class TSubclassOf;
+class UMovieSceneCustomBinding;
 struct FRichCurve;
 enum class EMovieSceneKeyInterpolation : uint8;
+struct FMovieSceneSequenceID;
 
 namespace UE::MovieScene
 {
@@ -226,6 +229,15 @@ public:
 	*/
 	static MOVIESCENE_API bool IsBoundToSpawnable(UMovieSceneSequence* Sequence, const FGuid& ObjectId, TSharedRef<const UE::MovieScene::FSharedPlaybackState> SharedPlaybackState, int32 BindingIndex = 0);
 
+	/*
+	* Attempts to create a new custom spawnable binding for the passed in UObject*. 
+	* Where possible, it is preferred to call FSequencerUtilities::CreateOrReplaceBinding as it handles more cases. This should only be called in cases where there is no editor or sequencer context.
+	* FactoryCreatedActor may be passed in as an alternative option for creating the binding in the case an actor factory was able to create an actor from this object.
+	*/
+
+	static MOVIESCENE_API FGuid TryCreateCustomSpawnableBinding(UMovieSceneSequence* Sequence, UObject* CustomBindingObject);
+	
+
 	/*
 	* Returns the single bound object currently bound to the given objectid and binding index (optional).
 	*/
@@ -257,6 +269,14 @@ public:
 	static MOVIESCENE_API const UClass* GetBoundObjectClass(UMovieSceneSequence* Sequence, const FGuid& ObjectId, int32 BindingIndex = 0);
 #endif
 
+	/* Returns a sorted list of all custom binding type classes currently known. Slow, may desire to cache result*/
+	static void MOVIESCENE_API GetPrioritySortedCustomBindingTypes(TArray<const TSubclassOf<UMovieSceneCustomBinding>>& OutCustomBindingTypes);
+
+	/* For cases where the user does not have a IMovieScenePlayer with a shared playback state, creates a transient one. Use sparingly. */
+	static TSharedRef<UE::MovieScene::FSharedPlaybackState> MOVIESCENE_API CreateTransientSharedPlaybackState(UObject* WorldContext, UMovieSceneSequence* Sequence);
+
+	/* Finds the resolution context to use to resolve the given guid. */
+	static MOVIESCENE_API UObject* GetResolutionContext(UMovieSceneSequence* Sequence, const FGuid& ObjectId, const FMovieSceneSequenceID& SequenceID, TSharedRef<const UE::MovieScene::FSharedPlaybackState> SharedPlaybackState);
 };
 
 /**
diff --git a/Engine/Source/Runtime/MovieScene/Public/MovieSceneDynamicBinding.h b/Engine/Source/Runtime/MovieScene/Public/MovieSceneDynamicBinding.h
--- a/Engine/Source/Runtime/MovieScene/Public/MovieSceneDynamicBinding.h
+++ b/Engine/Source/Runtime/MovieScene/Public/MovieSceneDynamicBinding.h
@@ -87,7 +87,7 @@ struct FMovieSceneDynamicBindingResolveParams
 };
 
 USTRUCT(BlueprintType)
-struct FMovieSceneDynamicBindingResolveResult
+struct MOVIESCENE_API FMovieSceneDynamicBindingResolveResult
 {
 	GENERATED_BODY()
 
diff --git a/Engine/Source/Runtime/MovieScene/Public/MovieScenePossessable.h b/Engine/Source/Runtime/MovieScene/Public/MovieScenePossessable.h
--- a/Engine/Source/Runtime/MovieScene/Public/MovieScenePossessable.h
+++ b/Engine/Source/Runtime/MovieScene/Public/MovieScenePossessable.h
@@ -154,9 +154,9 @@ public:
 	UPROPERTY(EditAnywhere, AdvancedDisplay, Category=Actor)
 	TArray<FName> Tags;
 
-	/** Optional user-defined possessable lookup information */
-	UPROPERTY(EditAnywhere, AdvancedDisplay, Category="Sequencer")
-	FMovieSceneDynamicBinding DynamicBinding;
+	// Property has been deprecated- 'Dynamic Bindings' are now handled by UMovieSceneReplaceableDirectorBlueprintBinding and the custom binding system
+	UPROPERTY()
+	FMovieSceneDynamicBinding DynamicBinding_DEPRECATED;
 
 	/* Get the optional binding id for binding to a spawnable */
 	const FMovieSceneObjectBindingID& GetSpawnableObjectBindingID() const 
diff --git a/Engine/Source/Runtime/MovieScene/Public/MovieSceneSequence.h b/Engine/Source/Runtime/MovieScene/Public/MovieSceneSequence.h
--- a/Engine/Source/Runtime/MovieScene/Public/MovieSceneSequence.h
+++ b/Engine/Source/Runtime/MovieScene/Public/MovieSceneSequence.h
@@ -41,6 +41,12 @@ struct FMovieScenePossessable;
 struct FMovieSceneTimecodeSource;
 struct FUniversalObjectLocator;
 struct FMovieSceneBindingReferences;
+struct FMovieSceneDynamicBinding;
+
+#if WITH_EDITOR
+struct FSlateBrush;
+#endif
+
 
 namespace UE::MovieScene
 {
@@ -139,6 +145,7 @@ public:
 	 * @param Params				Resolve parameters specifying the context and fragment-specific parameters
 	 * @param OutObjects			Destination array to add found objects to
 	 */
+	UE_DEPRECATED(5.5, "Please use the version that also takes a SharedPlaybackState")
 	MOVIESCENE_API void LocateBoundObjects(const FGuid& ObjectId, const UE::UniversalObjectLocator::FResolveParams& ResolveParams, TArray<UObject*, TInlineAllocator<1>>& OutObjects) const;
 
 	/**
@@ -149,13 +156,27 @@ public:
 	 * @param Context				Optional context to use to find the required object (for instance, a parent spawnable object)
 	 * @return An array of all bound objects
 	 */
+	UE_DEPRECATED(5.5, "Please use the version that also takes a SharedPlaybackState")
 	TArray<UObject*, TInlineAllocator<1>> LocateBoundObjects(const FGuid& ObjectId, const UE::UniversalObjectLocator::FResolveParams& Context) const
 	{
 		TArray<UObject*, TInlineAllocator<1>> OutObjects;
+		PRAGMA_DISABLE_DEPRECATION_WARNINGS
 		LocateBoundObjects(ObjectId, Context, OutObjects);
+		PRAGMA_ENABLE_DEPRECATION_WARNINGS
 		return OutObjects;
 	}
 
+	/**
+	 * Locate all the objects that correspond to the specified object ID, using the specified parameters.
+	 * Calling this directly instead of IMovieScenePlayer::ResolveBoundObjects means that you won't be able to locate Spawnable objects or objects from Custom Bindings.
+	 *
+	 * @param ObjectId				The unique identifier of the object.
+	 * @param ResolveParams			Resolve parameters specifying the context and fragment-specific parameters
+	 * @param SharedPlaybackState   Optional ptr to SharedPlaybackState
+	 * @param OutObjects			Destination array to add found objects to
+	 */
+	MOVIESCENE_API virtual void LocateBoundObjects(const FGuid& ObjectId, const UE::UniversalObjectLocator::FResolveParams& ResolveParams, TSharedPtr<const FSharedPlaybackState> SharedPlaybackState, TArray<UObject*, TInlineAllocator<1>>& OutObjects) const;
+
 	/**
 	 * Attempt to find the guid relating to the specified object
 	 *
@@ -302,6 +323,12 @@ public:
 	UFUNCTION(BlueprintCallable, Category = "Sequencer|Sequence")
 	MOVIESCENE_API FMovieSceneTimecodeSource GetEarliestTimecodeSource() const;
 
+	MOVIESCENE_API virtual void IterateDynamicBindings(const TSharedRef<UE::MovieScene::FSharedPlaybackState> SharedPlaybackState, TFunction<void(const FGuid&, FMovieSceneDynamicBinding&)> InCallback) {}
+	
+#if WITH_EDITOR
+	MOVIESCENE_API virtual const FSlateBrush* GetCustomBrushForBinding(FGuid BindingID) const { return nullptr; }
+#endif
+
 public:
 
 	MOVIESCENE_API virtual void PostLoad() override;
diff --git a/Engine/Source/Runtime/MovieScene/Public/MovieSceneSpawnable.h b/Engine/Source/Runtime/MovieScene/Public/MovieSceneSpawnable.h
--- a/Engine/Source/Runtime/MovieScene/Public/MovieSceneSpawnable.h
+++ b/Engine/Source/Runtime/MovieScene/Public/MovieSceneSpawnable.h
@@ -262,9 +262,9 @@ public:
 	UPROPERTY(EditAnywhere, Category=Actor)
 	bool bNetAddressableName;
 
-	/** Optional user-defined spawning information */
-	UPROPERTY(EditAnywhere, Category="Sequencer")
-	FMovieSceneDynamicBinding DynamicBinding;
+	// Property has been deprecated- 'Dynamic Bindings' are now handled by UMovieSceneReplaceableDirectorBlueprintBinding and the custom binding system
+	UPROPERTY()
+	FMovieSceneDynamicBinding DynamicBinding_DEPRECATED;
 
 private:
 
diff --git a/Engine/Source/Runtime/MovieSceneTracks/Private/Bindings/MovieSceneSpawnableActorBinding.cpp b/Engine/Source/Runtime/MovieSceneTracks/Private/Bindings/MovieSceneSpawnableActorBinding.cpp
--- a/Engine/Source/Runtime/MovieSceneTracks/Private/Bindings/MovieSceneSpawnableActorBinding.cpp
+++ b/Engine/Source/Runtime/MovieSceneTracks/Private/Bindings/MovieSceneSpawnableActorBinding.cpp
@@ -76,35 +76,11 @@ UObject* UMovieSceneSpawnableActorBindingBase::SpawnObjectInternal(UWorld* World
 		return nullptr;
 	}
 
-	// Ensure this spawnable is not a preview actor. Preview actors will not have BeginPlay() called on them.
-#if WITH_EDITOR
-	SpawnedActor->bIsEditorPreviewActor = false;
-#endif
-
-	// tag this actor so we know it was spawned by sequencer
-	SpawnedActor->Tags.AddUnique(SequencerActorTag);
-
 	if (bNetAddressableName)
 	{
 		SpawnedActor->SetNetAddressable();
 	}
 
-#if WITH_EDITOR
-	if (GIsEditor)
-	{
-		// Explicitly set RF_Transactional on spawned actors so we can undo/redo properties on them.
-		// This particular UObject will be marked RF_Transactional by the caller, but we need to set it on the components.
-
-		for (UActorComponent* Component : SpawnedActor->GetComponents())
-		{
-			if (Component)
-			{
-				Component->SetFlags(RF_Transactional);
-			}
-		}
-	}
-#endif
-
 	if (UMovieSceneEntitySystemLinker* Linker = SharedPlaybackState->GetLinker())
 	{
 		if (UMovieSceneDeferredComponentMovementSystem* DeferredMovementSystem = Linker->FindSystem<UMovieSceneDeferredComponentMovementSystem>())
@@ -122,33 +98,6 @@ UObject* UMovieSceneSpawnableActorBindingBase::SpawnObjectInternal(UWorld* World
 	const bool bIsDefaultTransform = true;
 	SpawnedActor->FinishSpawning(SpawnTransformToUse, bIsDefaultTransform);
 
-#if WITH_EDITOR
-	// Don't set the actor label in PIE as this requires flushing async loading.
-	if (WorldContext->WorldType == EWorldType::Editor)
-	{
-		FString BindingName = GetDesiredBindingName();
-		FString ActorLabel = !BindingName.IsEmpty() ? BindingName : SpawnName.ToString();
-
-		if (FMovieScenePossessable* Possessable = MovieScene.FindPossessable(BindingId))
-		{
-			ActorLabel = Possessable->GetName();
-
-			if (UMovieSceneSequence* Sequence = MovieScene.GetTypedOuter<UMovieSceneSequence>())
-			{
-				if (const FMovieSceneBindingReferences* BindingReferences = Sequence->GetBindingReferences())
-				{
-					if (BindingReferences->GetReferences(BindingId).Num() > 1)
-					{
-						// If there are multiple bound objects, use the Object Template actor label instead of the possessable name
-						ActorLabel = ActorTemplate->GetActorLabel();
-					}
-				}
-			}
-		}
-		SpawnedActor->SetActorLabel(ActorLabel);
-	}
-#endif
-
 	return SpawnedActor;
 }
 
@@ -160,21 +109,6 @@ void UMovieSceneSpawnableActorBindingBase::DestroySpawnedObjectInternal(UObject*
 		return;
 	}
 
-#if WITH_EDITOR
-	if (GIsEditor)
-	{
-		// Explicitly remove RF_Transactional on spawned actors since we don't want to trasact spawn/destroy events
-		// This particular UObject will have RF_Transactional cleared by the caller, but we need to cleared it on the components.
-		for (UActorComponent* Component : Actor->GetComponents())
-		{
-			if (Component)
-			{
-				Component->ClearFlags(RF_Transactional);
-			}
-		}
-	}
-#endif
-
 	UWorld* World = Actor->GetWorld();
 	if (World)
 	{
@@ -453,7 +387,7 @@ UMovieSceneCustomBinding* UMovieSceneSpawnableActorBinding::CreateNewCustomBindi
 
 	auto CreateBinding = [&]()
 	{
-		return NewObject<UMovieSceneSpawnableActorBinding>(&OwnerMovieScene, UMovieSceneSpawnableActorBinding::StaticClass(), InstancedBindingName, RF_Transactional);
+		return NewObject<UMovieSceneSpawnableActorBinding>(&OwnerMovieScene, GetClass(), InstancedBindingName, RF_Transactional);
 	};
 
 	// Deal with creating a spawnable from an instance of an actor
diff --git a/Engine/Source/Runtime/MovieSceneTracks/Public/Bindings/MovieSceneReplaceableActorBinding.h b/Engine/Source/Runtime/MovieSceneTracks/Public/Bindings/MovieSceneReplaceableActorBinding.h
--- a/Engine/Source/Runtime/MovieSceneTracks/Public/Bindings/MovieSceneReplaceableActorBinding.h
+++ b/Engine/Source/Runtime/MovieSceneTracks/Public/Bindings/MovieSceneReplaceableActorBinding.h
@@ -41,5 +41,7 @@ protected:
 
 	MOVIESCENETRACKS_API virtual TSubclassOf<UMovieSceneSpawnableBindingBase> GetInnerSpawnableClass() const override;
 
+	virtual int32 GetCustomBindingPriority() const override { return 9; }
+
 };
 
diff --git a/Engine/Source/Runtime/MovieSceneTracks/Public/Bindings/MovieSceneSpawnableActorBinding.h b/Engine/Source/Runtime/MovieSceneTracks/Public/Bindings/MovieSceneSpawnableActorBinding.h
--- a/Engine/Source/Runtime/MovieSceneTracks/Public/Bindings/MovieSceneSpawnableActorBinding.h
+++ b/Engine/Source/Runtime/MovieSceneTracks/Public/Bindings/MovieSceneSpawnableActorBinding.h
@@ -83,9 +83,10 @@ protected:
 	/* Returns the transform to spawn the actor at*/
 	MOVIESCENETRACKS_API  FTransform GetSpawnTransform() const;
 
+	MOVIESCENETRACKS_API UClass* GetBoundObjectClass() const override { return GetActorClass(); }
 #if WITH_EDITOR
 	/* MovieSceneCustomBinding overrides*/
-	MOVIESCENETRACKS_API UClass* GetBoundObjectClass() const override { return GetActorClass(); }
+	virtual int32 GetCustomBindingPriority() const override { return 10; }
 #endif
 
 protected:
@@ -119,33 +120,33 @@ public:
 public:
 
 	/* MovieSceneSpawnableActorBindingBase overrides*/
-	bool SupportsObjectTemplates() const override { return true; }
+	MOVIESCENETRACKS_API bool SupportsObjectTemplates() const override { return true; }
 
 	/* Override and return the object template if the binding type supports object templates*/
-	UObject* GetObjectTemplate() override { return ActorTemplate.Get(); }
+	MOVIESCENETRACKS_API UObject* GetObjectTemplate() override { return ActorTemplate.Get(); }
 
 	/* MovieSceneSpawnableBindingBase overrides*/
-	void SetObjectTemplate(UObject* InObjectTemplate) override;
-	void CopyObjectTemplate(UObject* InSourceObject, UMovieSceneSequence& MovieSceneSequence) override;
+	MOVIESCENETRACKS_API void SetObjectTemplate(UObject* InObjectTemplate) override;
+	MOVIESCENETRACKS_API void CopyObjectTemplate(UObject* InSourceObject, UMovieSceneSequence& MovieSceneSequence) override;
 
 	/* MovieSceneCustomBinding overrides*/
-	bool SupportsBindingCreationFromObject(const UObject* SourceObject) const override;
-	UMovieSceneCustomBinding* CreateNewCustomBinding(UObject* SourceObject, UMovieScene& OwnerMovieScene) override;
+	MOVIESCENETRACKS_API virtual bool SupportsBindingCreationFromObject(const UObject* SourceObject) const override;
+	MOVIESCENETRACKS_API virtual UMovieSceneCustomBinding* CreateNewCustomBinding(UObject* SourceObject, UMovieScene& OwnerMovieScene) override;
 
 #if WITH_EDITOR
-	bool SupportsConversionFromBinding(const FMovieSceneBindingReference& BindingReference, const UObject* SourceObject) const override;
-	UMovieSceneCustomBinding* CreateCustomBindingFromBinding(const FMovieSceneBindingReference& BindingReference, UObject* SourceObject, UMovieScene& OwnerMovieScene) override;
-	FText GetBindingTypePrettyName() const override;
+	MOVIESCENETRACKS_API virtual bool SupportsConversionFromBinding(const FMovieSceneBindingReference& BindingReference, const UObject* SourceObject) const override;
+	MOVIESCENETRACKS_API virtual UMovieSceneCustomBinding* CreateCustomBindingFromBinding(const FMovieSceneBindingReference& BindingReference, UObject* SourceObject, UMovieScene& OwnerMovieScene) override;
+	MOVIESCENETRACKS_API virtual FText GetBindingTypePrettyName() const override;
 #endif
 
 protected:
 
 	/* MovieSceneSpawnableBindingBase overrides*/
-	UWorld* GetWorldContext(TSharedRef<const UE::MovieScene::FSharedPlaybackState> SharedPlaybackState) const override;
+	MOVIESCENETRACKS_API UWorld* GetWorldContext(TSharedRef<const UE::MovieScene::FSharedPlaybackState> SharedPlaybackState) const override;
 	
 	/* MovieSceneSpawnableActorBindingBase overrides*/
-	TSubclassOf<AActor> GetActorClass() const override;
-	AActor* GetActorTemplate() const override { return ActorTemplate; }
+	MOVIESCENETRACKS_API TSubclassOf<AActor> GetActorClass() const override;
+	MOVIESCENETRACKS_API AActor* GetActorTemplate() const override { return ActorTemplate; }
 
 
 private:
diff --git a/Engine/Source/Runtime/UMG/Private/Animation/WidgetAnimation.cpp b/Engine/Source/Runtime/UMG/Private/Animation/WidgetAnimation.cpp
--- a/Engine/Source/Runtime/UMG/Private/Animation/WidgetAnimation.cpp
+++ b/Engine/Source/Runtime/UMG/Private/Animation/WidgetAnimation.cpp
@@ -1,4 +1,4 @@
-// Copyright Epic Games, Inc. All Rights Reserved.
+	// Copyright Epic Games, Inc. All Rights Reserved.
 
 #include "Animation/WidgetAnimation.h"
 #include "UObject/Package.h"
@@ -12,11 +12,18 @@
 #include "Tracks/MovieSceneMaterialParameterCollectionTrack.h"
 #include "UObject/SequencerObjectVersion.h"
 
+#if WITH_EDITOR
+#include "Styling/AppStyle.h"
+#endif
+
 #include UE_INLINE_GENERATED_CPP_BY_NAME(WidgetAnimation)
 
 
 #define LOCTEXT_NAMESPACE "UWidgetAnimation"
 
+#if WITH_EDITOR
+UWidgetAnimation::FFixupWidgetDynamicBindingsEvent UWidgetAnimation::FixupWidgetDynamicBindingsEvent;
+#endif
 
 /* UWidgetAnimation structors
  *****************************************************************************/
@@ -39,6 +46,35 @@ void UWidgetAnimation::PostLoad()
 		bLegacyFinishOnStop = false;
 	}
 
+	if (GetLinkerCustomVersion(FFortniteMainBranchObjectVersion::GUID) < FFortniteMainBranchObjectVersion::LevelSequenceUpgradeDynamicBindings)
+	{
+		bool bConvertedDynamicBinding = false;
+		for (int32 Index = 0; Index < MovieScene->GetPossessableCount(); ++Index)
+		{
+			FMovieScenePossessable& Possessable = MovieScene->GetPossessable(Index);
+			if (Possessable.DynamicBinding_DEPRECATED.Function)
+			{
+				bConvertedDynamicBinding = true;
+
+				for (FWidgetAnimationBinding& Binding : AnimationBindings)
+				{
+					if (Binding.AnimationGuid == Possessable.GetGuid())
+					{
+						Binding.DynamicBinding = Possessable.DynamicBinding_DEPRECATED;
+						Possessable.DynamicBinding_DEPRECATED = FMovieSceneDynamicBinding();
+						break;
+					}
+				}
+			}
+		}
+#if WITH_EDITOR
+		if (bConvertedDynamicBinding && FixupWidgetDynamicBindingsEvent.IsBound())
+		{
+			FixupWidgetDynamicBindingsEvent.Broadcast(this);
+		}
+#endif
+	}
+
 	Super::PostLoad();
 }
 
@@ -87,6 +123,22 @@ ETrackSupport UWidgetAnimation::IsTrackSupported(TSubclassOf<class UMovieSceneTr
 
 	return Super::IsTrackSupported(InTrackClass);
 }
+
+const FSlateBrush* UWidgetAnimation::GetCustomBrushForBinding(FGuid BindingID) const
+{
+	for (const FWidgetAnimationBinding& Binding : AnimationBindings)
+	{
+		if (Binding.AnimationGuid == BindingID)
+		{
+			if (Binding.DynamicBinding.Function)
+			{
+				return FAppStyle::GetBrush("Sequencer.DynamicBindingIconOverlay");
+			}
+		}
+	}
+	return nullptr;
+}
+
 #endif
 
 float UWidgetAnimation::GetStartTime() const
@@ -225,14 +277,14 @@ bool UWidgetAnimation::CanPossessObject(UObject& Object, UObject* InPlaybackCont
 	return (Object.IsA<UVisual>() && Object.IsIn(PreviewWidget));
 }
 
-void UWidgetAnimation::LocateBoundObjects(const FGuid& ObjectId, UObject* InContext, TArray<UObject*, TInlineAllocator<1>>& OutObjects) const
+void UWidgetAnimation::LocateBoundObjects(const FGuid& ObjectId, const UE::UniversalObjectLocator::FResolveParams& ResolveParams, TSharedPtr<const FSharedPlaybackState> SharedPlaybackState, TArray<UObject*, TInlineAllocator<1>>& OutObjects) const
 {
-	if (InContext == nullptr)
+	if (ResolveParams.Context == nullptr)
 	{
 		return;
 	}
 
-	UUserWidget* PreviewWidget = CastChecked<UUserWidget>(InContext);
+	UUserWidget* PreviewWidget = CastChecked<UUserWidget>(ResolveParams.Context);
 	if (PreviewWidget->WidgetTree == nullptr)
 	{
 		return;
@@ -242,7 +294,7 @@ void UWidgetAnimation::LocateBoundObjects(const FGuid& ObjectId, UObject* InCont
 	{
 		if (Binding.AnimationGuid == ObjectId)
 		{
-			UObject* FoundObject = Binding.FindRuntimeObject(*PreviewWidget->WidgetTree, *PreviewWidget);
+			UObject* FoundObject = Binding.FindRuntimeObject(*PreviewWidget->WidgetTree, *PreviewWidget, this, SharedPlaybackState);
 
 			if (FoundObject)
 			{
@@ -250,10 +302,8 @@ void UWidgetAnimation::LocateBoundObjects(const FGuid& ObjectId, UObject* InCont
 			}
 		}
 	}
-
 }
 
-
 UMovieScene* UWidgetAnimation::GetMovieScene() const
 {
 	return MovieScene;
@@ -266,6 +316,14 @@ UObject* UWidgetAnimation::CreateDirectorInstance(TSharedRef<const FSharedPlayba
 	return WidgetContext;
 }
 
+void UWidgetAnimation::IterateDynamicBindings(const TSharedRef<UE::MovieScene::FSharedPlaybackState> SharedPlaybackState, TFunction<void(const FGuid&, FMovieSceneDynamicBinding&)> InCallback)
+{
+	for (FWidgetAnimationBinding& AnimationBinding : AnimationBindings)
+	{
+		InCallback(AnimationBinding.AnimationGuid, AnimationBinding.DynamicBinding);
+	}
+}
+
 UObject* UWidgetAnimation::GetParentObject(UObject* Object) const
 {
 	UPanelSlot* Slot = Cast<UPanelSlot>(Object);
diff --git a/Engine/Source/Runtime/UMG/Private/Animation/WidgetAnimationBinding.cpp b/Engine/Source/Runtime/UMG/Private/Animation/WidgetAnimationBinding.cpp
--- a/Engine/Source/Runtime/UMG/Private/Animation/WidgetAnimationBinding.cpp
+++ b/Engine/Source/Runtime/UMG/Private/Animation/WidgetAnimationBinding.cpp
@@ -4,6 +4,7 @@
 #include "UObject/Object.h"
 #include "Components/Widget.h"
 #include "Blueprint/WidgetTree.h"
+#include "MovieSceneDynamicBindingInvoker.h"
 
 #include UE_INLINE_GENERATED_CPP_BY_NAME(WidgetAnimationBinding)
 
@@ -13,6 +14,20 @@
 
 UObject* FWidgetAnimationBinding::FindRuntimeObject(const UWidgetTree& WidgetTree, UUserWidget& UserWidget ) const
 {	
+	return FindRuntimeObject(WidgetTree, UserWidget, nullptr, nullptr);
+}
+
+UObject* FWidgetAnimationBinding::FindRuntimeObject(const UWidgetTree& WidgetTree, UUserWidget& UserWidget, const UMovieSceneSequence* Sequence, TSharedPtr<const UE::MovieScene::FSharedPlaybackState> SharedPlaybackState) const
+{
+	if (Sequence && SharedPlaybackState && DynamicBinding.Function)
+	{
+		FMovieSceneDynamicBindingResolveResult ResolveResult = FMovieSceneDynamicBindingInvoker::ResolveDynamicBinding(SharedPlaybackState.ToSharedRef(), const_cast<UMovieSceneSequence*>(Sequence), MovieSceneSequenceID::Root, AnimationGuid, DynamicBinding);
+		if (ResolveResult.Object)
+		{
+			return ResolveResult.Object;
+		}
+	}
+
 	if (bIsRootWidget)
 	{
 		return &UserWidget;
@@ -32,4 +47,4 @@ UObject* FWidgetAnimationBinding::FindRuntimeObject(const UWidgetTree& WidgetTre
 	}
 
 	return FoundObject;
-}
+}
\ No newline at end of file
