Sequencer: Remove unused function that is triggering duplicate symbol errors in some builds.

diff --git a/Engine/Source/Runtime/LevelSequence/Private/LevelSequence.cpp b/Engine/Source/Runtime/LevelSequence/Private/LevelSequence.cpp
--- a/Engine/Source/Runtime/LevelSequence/Private/LevelSequence.cpp
+++ b/Engine/Source/Runtime/LevelSequence/Private/LevelSequence.cpp
@@ -444,29 +444,6 @@ void ULevelSequence::PostLoad()
 
 #if WITH_EDITOR
 
-bool UpdateFolderBindingID(UMovieSceneFolder* Folder, FGuid OldGuid, FGuid NewGuid)
-{
-	for (FGuid ChildGuid : Folder->GetChildObjectBindings())
-	{
-		if (ChildGuid == OldGuid)
-		{
-			Folder->AddChildObjectBinding(NewGuid);
-			Folder->RemoveChildObjectBinding(OldGuid);
-			return true;
-		}
-	}
-
-	for (UMovieSceneFolder* ChildFolder : Folder->GetChildFolders())
-	{
-		if (UpdateFolderBindingID(ChildFolder, OldGuid, NewGuid))
-		{
-			return true;
-		}
-	}
-
-	return false;
-}
-
 bool ULevelSequence::ConvertOldSpawnables()
 {
 	bool bConvertedDynamicBinding = false;
