Sequencer: Add to ULevelSequence::DeclareConstructClasses as we're not constructing other objects during PostLoad.

diff --git a/Engine/Source/Runtime/LevelSequence/Private/LevelSequence.cpp b/Engine/Source/Runtime/LevelSequence/Private/LevelSequence.cpp
--- a/Engine/Source/Runtime/LevelSequence/Private/LevelSequence.cpp
+++ b/Engine/Source/Runtime/LevelSequence/Private/LevelSequence.cpp
@@ -580,6 +580,8 @@ void ULevelSequence::DeclareConstructClasses(TArray<FTopLevelAssetPath>& OutCons
 {
 	Super::DeclareConstructClasses(OutConstructClasses, SpecificSubclass);
 	OutConstructClasses.Add(FTopLevelAssetPath(UObjectRedirector::StaticClass()));
+	OutConstructClasses.Add(FTopLevelAssetPath(UMovieSceneSpawnableActorBinding::StaticClass()));
+	OutConstructClasses.Add(FTopLevelAssetPath(UMovieSceneSpawnableDirectorBlueprintBinding::StaticClass()));
 }
 #endif
 
