Sequencer: revert validation of dynamically bound objects, change assert into error log

The validation of bound objects was only valid for level sequences, and broke things for UMG animations.
To fix the original problem, we replace the assert with an error log and just accept "wrong" objects without crashing.

diff --git a/Engine/Source/Runtime/LevelSequence/Private/LevelSequenceSpawnRegister.cpp b/Engine/Source/Runtime/LevelSequence/Private/LevelSequenceSpawnRegister.cpp
--- a/Engine/Source/Runtime/LevelSequence/Private/LevelSequenceSpawnRegister.cpp
+++ b/Engine/Source/Runtime/LevelSequence/Private/LevelSequenceSpawnRegister.cpp
@@ -49,7 +49,10 @@ void FLevelSequenceSpawnRegister::DestroySpawnedObject(UObject& Object, UMovieSc
 			}
 		}
 
-		checkf(false, TEXT("No valid object spawner found to destroy spawned object of type %s"), *Object.GetClass()->GetName());
+		UE_LOG(
+			LogMovieScene, Error,
+			TEXT("No valid object spawner found to destroy spawned object '%s' of type '%s'."),
+			*Object.GetPathName(), *Object.GetClass()->GetName());
 	}
 }
 
