Sequencer: Fix issues where CreateTransientPlaybackState fails due to being called during a level load whene there is no editor world.

diff --git a/Engine/Source/Editor/MovieSceneTools/Public/MovieSceneDynamicBindingUtils.h b/Engine/Source/Editor/MovieSceneTools/Public/MovieSceneDynamicBindingUtils.h
--- a/Engine/Source/Editor/MovieSceneTools/Public/MovieSceneDynamicBindingUtils.h
+++ b/Engine/Source/Editor/MovieSceneTools/Public/MovieSceneDynamicBindingUtils.h
@@ -39,10 +39,10 @@ struct MOVIESCENETOOLS_API FMovieSceneDynamicBindingUtils
 	static void IterateDynamicBindings(UMovieScene* InMovieScene, Callback&& InCallback)
 	{
 		UMovieSceneSequence* ThisSequence = InMovieScene->GetTypedOuter<UMovieSceneSequence>();
-		TSharedRef<UE::MovieScene::FSharedPlaybackState> TransientPlaybackState = MovieSceneHelpers::CreateTransientSharedPlaybackState(GEditor->GetEditorWorldContext().World(), ThisSequence);
-		ThisSequence->IterateDynamicBindings(TransientPlaybackState, InCallback);
-	}
 
+		ThisSequence->IterateDynamicBindings(InCallback);
+	}
+	
 	/**
 	 * Utility function for gathering all dynamic bindings in a sequence into a container.
 	 */
diff --git a/Engine/Source/Runtime/LevelSequence/Private/LevelSequence.cpp b/Engine/Source/Runtime/LevelSequence/Private/LevelSequence.cpp
--- a/Engine/Source/Runtime/LevelSequence/Private/LevelSequence.cpp
+++ b/Engine/Source/Runtime/LevelSequence/Private/LevelSequence.cpp
@@ -961,7 +961,8 @@ const TArray<UAssetUserData*>* ULevelSequence::GetAssetUserDataArray() const
 	return &ToRawPtrTArrayUnsafe(AssetUserData);
 }
 
-void ULevelSequence::IterateDynamicBindings(const TSharedRef<UE::MovieScene::FSharedPlaybackState> SharedPlaybackState, TFunction<void(const FGuid&, FMovieSceneDynamicBinding&)> InCallback)
+#if WITH_EDITOR
+void ULevelSequence::IterateDynamicBindings(TFunction<void(const FGuid&, FMovieSceneDynamicBinding&)> InCallback)
 {
 	for (FMovieSceneBindingReference& BindingReference : BindingReferences.GetAllReferences())
 	{
@@ -972,10 +973,23 @@ void ULevelSequence::IterateDynamicBindings(const TSharedRef<UE::MovieScene::FSh
 				InCallback(BindingReference.ID, ReplaceableDirectorBlueprintBinding->DynamicBinding);
 			}
 
-			if (UMovieSceneSpawnableDirectorBlueprintBinding* SpawnableDirectorBlueprintBinding = Cast<UMovieSceneSpawnableDirectorBlueprintBinding>(BindingReference.CustomBinding->AsSpawnable(SharedPlaybackState)))
+			// We can't use 'AsSpawnable' here because we don't have playback state and we might not have a world context. 
+			// This should only be called from an editor context though, so we can just check the inner spawnable.
+
+			// If the binding is itself a spawnable director blueprint binding, then iterate over it
+			if (UMovieSceneSpawnableDirectorBlueprintBinding* SpawnableDirectorBlueprintBinding = Cast<UMovieSceneSpawnableDirectorBlueprintBinding>(BindingReference.CustomBinding))
 			{
 				InCallback(BindingReference.ID, SpawnableDirectorBlueprintBinding->DynamicBinding);
 			}
+			else if (UMovieSceneReplaceableBindingBase* ReplaceableBinding = Cast<UMovieSceneReplaceableBindingBase>(BindingReference.CustomBinding))
+			{
+				if (UMovieSceneSpawnableDirectorBlueprintBinding* InnerSpawnableDirectorBlueprintBinding = Cast<UMovieSceneSpawnableDirectorBlueprintBinding>(ReplaceableBinding->PreviewSpawnable))
+				{
+					InCallback(BindingReference.ID, InnerSpawnableDirectorBlueprintBinding->DynamicBinding);
+				}
+			}
 		}
 	}
 }
+
+#endif
diff --git a/Engine/Source/Runtime/LevelSequence/Public/LevelSequence.h b/Engine/Source/Runtime/LevelSequence/Public/LevelSequence.h
--- a/Engine/Source/Runtime/LevelSequence/Public/LevelSequence.h
+++ b/Engine/Source/Runtime/LevelSequence/Public/LevelSequence.h
@@ -102,10 +102,10 @@ public:
 	UE_DEPRECATED(5.4, "Use the base class LocateBoundObjects()")
 	LEVELSEQUENCE_API void LocateBoundObjects(const FGuid& ObjectId, UObject* Context, const FLevelSequenceBindingReference::FResolveBindingParams& InResolveBindingParams, TArray<UObject*, TInlineAllocator<1>>& OutObjects) const;
 
-	LEVELSEQUENCE_API virtual void IterateDynamicBindings(const TSharedRef<UE::MovieScene::FSharedPlaybackState> SharedPlaybackState, TFunction<void(const FGuid&, FMovieSceneDynamicBinding&)> InCallback) override;
-
+	
 #if WITH_EDITOR
 
+	LEVELSEQUENCE_API virtual void IterateDynamicBindings(TFunction<void(const FGuid&, FMovieSceneDynamicBinding&)> InCallback) override;
 
 public:
 
diff --git a/Engine/Source/Runtime/MovieScene/Public/MovieSceneSequence.h b/Engine/Source/Runtime/MovieScene/Public/MovieSceneSequence.h
--- a/Engine/Source/Runtime/MovieScene/Public/MovieSceneSequence.h
+++ b/Engine/Source/Runtime/MovieScene/Public/MovieSceneSequence.h
@@ -323,9 +323,9 @@ public:
 	UFUNCTION(BlueprintCallable, Category = "Sequencer|Sequence")
 	MOVIESCENE_API FMovieSceneTimecodeSource GetEarliestTimecodeSource() const;
 
-	MOVIESCENE_API virtual void IterateDynamicBindings(const TSharedRef<UE::MovieScene::FSharedPlaybackState> SharedPlaybackState, TFunction<void(const FGuid&, FMovieSceneDynamicBinding&)> InCallback) {}
-	
 #if WITH_EDITOR
+	MOVIESCENE_API virtual void IterateDynamicBindings(TFunction<void(const FGuid&, FMovieSceneDynamicBinding&)> InCallback) {}
+
 	MOVIESCENE_API virtual const FSlateBrush* GetCustomBrushForBinding(FGuid BindingID) const { return nullptr; }
 #endif
 
diff --git a/Engine/Source/Runtime/UMG/Private/Animation/WidgetAnimation.cpp b/Engine/Source/Runtime/UMG/Private/Animation/WidgetAnimation.cpp
--- a/Engine/Source/Runtime/UMG/Private/Animation/WidgetAnimation.cpp
+++ b/Engine/Source/Runtime/UMG/Private/Animation/WidgetAnimation.cpp
@@ -139,6 +139,15 @@ const FSlateBrush* UWidgetAnimation::GetCustomBrushForBinding(FGuid BindingID) c
 	return nullptr;
 }
 
+
+void UWidgetAnimation::IterateDynamicBindings(TFunction<void(const FGuid&, FMovieSceneDynamicBinding&)> InCallback)
+{
+	for (FWidgetAnimationBinding& AnimationBinding : AnimationBindings)
+	{
+		InCallback(AnimationBinding.AnimationGuid, AnimationBinding.DynamicBinding);
+	}
+}
+
 #endif
 
 float UWidgetAnimation::GetStartTime() const
@@ -316,14 +325,6 @@ UObject* UWidgetAnimation::CreateDirectorInstance(TSharedRef<const FSharedPlayba
 	return WidgetContext;
 }
 
-void UWidgetAnimation::IterateDynamicBindings(const TSharedRef<UE::MovieScene::FSharedPlaybackState> SharedPlaybackState, TFunction<void(const FGuid&, FMovieSceneDynamicBinding&)> InCallback)
-{
-	for (FWidgetAnimationBinding& AnimationBinding : AnimationBindings)
-	{
-		InCallback(AnimationBinding.AnimationGuid, AnimationBinding.DynamicBinding);
-	}
-}
-
 UObject* UWidgetAnimation::GetParentObject(UObject* Object) const
 {
 	UPanelSlot* Slot = Cast<UPanelSlot>(Object);
diff --git a/Engine/Source/Runtime/UMG/Public/Animation/WidgetAnimation.h b/Engine/Source/Runtime/UMG/Public/Animation/WidgetAnimation.h
--- a/Engine/Source/Runtime/UMG/Public/Animation/WidgetAnimation.h
+++ b/Engine/Source/Runtime/UMG/Public/Animation/WidgetAnimation.h
@@ -104,8 +104,9 @@ public:
 	virtual void LocateBoundObjects(const FGuid& ObjectId, const UE::UniversalObjectLocator::FResolveParams& ResolveParams, TSharedPtr<const FSharedPlaybackState> SharedPlaybackState, TArray<UObject*, TInlineAllocator<1>>& OutObjects) const override;
 	virtual UObject* CreateDirectorInstance(TSharedRef<const FSharedPlaybackState> SharedPlaybackState, FMovieSceneSequenceID SequenceID) override;
 
-	virtual void IterateDynamicBindings(const TSharedRef<UE::MovieScene::FSharedPlaybackState> SharedPlaybackState, TFunction<void(const FGuid&, FMovieSceneDynamicBinding&)> InCallback) override;
+	
 #if WITH_EDITOR
+	virtual void IterateDynamicBindings(TFunction<void(const FGuid&, FMovieSceneDynamicBinding&)> InCallback) override;
 	virtual ETrackSupport IsTrackSupported(TSubclassOf<class UMovieSceneTrack> InTrackClass) const override;
 	virtual const FSlateBrush* GetCustomBrushForBinding(FGuid BindingID) const override;
 
