Sequencer: Crash after re-opening a Sequence with a Director Blueprint in Level Recorder.
There was an issue with director blueprint redirectors, which we shouldn't really be creating in PostLoad like this.

diff --git a/Engine/Source/Runtime/LevelSequence/Private/LevelSequence.cpp b/Engine/Source/Runtime/LevelSequence/Private/LevelSequence.cpp
--- a/Engine/Source/Runtime/LevelSequence/Private/LevelSequence.cpp
+++ b/Engine/Source/Runtime/LevelSequence/Private/LevelSequence.cpp
@@ -354,9 +354,9 @@ void ULevelSequence::PostLoad()
 		DirectorBlueprint->OnCompiled().RemoveAll(this);
 		DirectorBlueprint->OnCompiled().AddUObject(this, &ULevelSequence::OnDirectorRecompiled);
 
-		if (DirectorBlueprint->Rename(*GetDirectorBlueprintName(), nullptr, (REN_NonTransactional|REN_ForceNoResetLoaders|REN_DoNotDirty|REN_Test)))
+		if (DirectorBlueprint->Rename(*GetDirectorBlueprintName(), nullptr, (REN_NonTransactional|REN_ForceNoResetLoaders|REN_DoNotDirty|REN_Test|REN_DontCreateRedirectors)))
 		{
-			DirectorBlueprint->Rename(*GetDirectorBlueprintName(), nullptr, (REN_NonTransactional|REN_ForceNoResetLoaders|REN_DoNotDirty));
+			DirectorBlueprint->Rename(*GetDirectorBlueprintName(), nullptr, (REN_NonTransactional|REN_ForceNoResetLoaders|REN_DoNotDirty|REN_DontCreateRedirectors));
 		}
 	}
 
