- Fix for level sequence not working in injected levels: Submitting on behalf of Jean-Francois Dube

diff --git a/Engine/Source/Runtime/LevelSequence/Private/LevelSequenceActor.cpp b/Engine/Source/Runtime/LevelSequence/Private/LevelSequenceActor.cpp
--- a/Engine/Source/Runtime/LevelSequence/Private/LevelSequenceActor.cpp
+++ b/Engine/Source/Runtime/LevelSequence/Private/LevelSequenceActor.cpp
@@ -182,18 +182,7 @@ void ALevelSequenceActor::PreInitializeComponents()
 		ULevelStreaming* LevelStreaming = ULevelStreaming::FindStreamingLevel(Level);
 		if (LevelStreaming)
 		{
-			// Sub world partitions as always loaded cells + traditional level streaming
-			if (Level->IsWorldPartitionRuntimeCell())
-			{
-				StreamingWorld = LevelStreaming->GetStreamingWorld();
-				check(StreamingWorld);
-
-				LevelStreaming = ULevelStreaming::FindStreamingLevel(StreamingWorld->PersistentLevel);
-			}
-			else
-			{
-				StreamingWorld = Level->GetTypedOuter<UWorld>();
-			}
+			StreamingWorld = Level->GetTypedOuter<UWorld>();
 		}
 
 		if (LevelStreaming)
