Sequencer: Optimize assigning actor labels for newly spawned actors

- Profiling Sequence operations where there are large numbers of spawnables highlights the cost of setting actor labels for several reasons but principally because it triggers an event that allows many sub-systems to react to the change of actor label. The World Outliner and Level Editor Sequencer Integration both do this and perform expensive fixup tasks inline. This is not necessary for the initial spawn operation, but there was no affordance in the API to supply an actor label at spawn time: it was only possible to call SetActorLabel after spawn.

- This change introduces a new function, AActor::SetInitialActorLabel which is only valid to be called if bFinishedSpawning is false. Furthermore, a spawn parameter has been added that allows callees to specify an actor label at spawn time without incurring the overhead of a subsequent call to SetActorLabel. The new function is much more lightweight in that it does not do any PostEditChange events or trigger the label changed event.

diff --git a/Engine/Source/Editor/Sequencer/Private/LevelEditorSequencerIntegration.cpp b/Engine/Source/Editor/Sequencer/Private/LevelEditorSequencerIntegration.cpp
--- a/Engine/Source/Editor/Sequencer/Private/LevelEditorSequencerIntegration.cpp
+++ b/Engine/Source/Editor/Sequencer/Private/LevelEditorSequencerIntegration.cpp
@@ -386,8 +386,13 @@ void RenameBindingRecursive(FSequencer* Sequencer, UMovieScene* MovieScene, FMov
 
 void FLevelEditorSequencerIntegration::OnActorLabelChanged(AActor* ChangedActor)
 {
-	// Sync up the spawnable or possessable name with the actor label if it is changed in editor (but not in PIE)
-	const bool bPIEWorld = ChangedActor != nullptr && ChangedActor->GetWorld() != nullptr && ChangedActor->GetWorld()->IsPlayInEditor();
+	// Sync up the spawnable or possessable name with the actor label if it is changed in editor (but not in PIE or for preview actors)
+	if (ChangedActor == nullptr || ChangedActor->bIsEditorPreviewActor)
+	{
+		return;
+	}
+
+	const bool bPIEWorld = ChangedActor->GetWorld() != nullptr && ChangedActor->GetWorld()->IsPlayInEditor();
 	if (bPIEWorld)
 	{
 		return;
diff --git a/Engine/Source/Runtime/Engine/Classes/Engine/World.h b/Engine/Source/Runtime/Engine/Classes/Engine/World.h
--- a/Engine/Source/Runtime/Engine/Classes/Engine/World.h
+++ b/Engine/Source/Runtime/Engine/Classes/Engine/World.h
@@ -433,6 +433,9 @@ struct ENGINE_API FActorSpawnParameters
 	/* The UPackage to set the Actor in. If left as NULL the Package will not be set and the actor will be saved in the same package as the persistent level. */
 	class	UPackage* OverridePackage;
 
+	/** Optional initial editor label to apply to the actor */
+	FStringView InitialActorLabel;
+
 	/** The Guid to set to this actor. Should only be set when reinstancing blueprint actors. */
 	FGuid	OverrideActorGuid;
 #endif
diff --git a/Engine/Source/Runtime/Engine/Classes/GameFramework/Actor.h b/Engine/Source/Runtime/Engine/Classes/GameFramework/Actor.h
--- a/Engine/Source/Runtime/Engine/Classes/GameFramework/Actor.h
+++ b/Engine/Source/Runtime/Engine/Classes/GameFramework/Actor.h
@@ -2699,6 +2699,9 @@ public:
 	/** Advanced - clear the actor label. */
 	ENGINE_API void ClearActorLabel();
 
+	/** Advanced - only to be called by UWorld::SpawnActor to set an initial label for this actor */
+	void SetInitialActorLabel(FStringView InitialLabel);
+
 	/**
 	 * Returns if this actor's current label is editable.  Actor labels are only available in development builds.
 	 * @return	The editable status of the actor's label
diff --git a/Engine/Source/Runtime/Engine/Private/ActorEditor.cpp b/Engine/Source/Runtime/Engine/Private/ActorEditor.cpp
--- a/Engine/Source/Runtime/Engine/Private/ActorEditor.cpp
+++ b/Engine/Source/Runtime/Engine/Private/ActorEditor.cpp
@@ -1248,23 +1248,36 @@ void AActor::SetActorLabel(const FString& NewActorLabelDirty, bool bMarkDirty)
 		//Invalid actor name
 		UE_LOG(LogActor, Warning, TEXT("SetActorLabel failed: %s"), *OutErrorMessage.ToString());
 	}
-	else
+	// Only update the label if it has actually changed. We pass 'false' to GetActorLabel to prevent it
+	//    creating a redundant name and broadcasting PostEditChangeProperty and OnActorLabelChanged unnecessarily
+	else if (FCString::Strcmp(*NewActorLabel, *GetActorLabel(false)) != 0)
 	{
-		// First, update the actor label
-		{
-			// Has anything changed?
-			if (FCString::Strcmp(*NewActorLabel, *GetActorLabel()) != 0)
-			{
-				// Store new label
-				Modify(bMarkDirty);
-				ActorLabel = MoveTemp(NewActorLabel);
+		// Store new label
+		Modify(bMarkDirty);
+		ActorLabel = MoveTemp(NewActorLabel);
 
-				FPropertyChangedEvent PropertyEvent(FindFProperty<FProperty>(AActor::StaticClass(), "ActorLabel"));
-				PostEditChangeProperty(PropertyEvent);
+		FPropertyChangedEvent PropertyEvent(FindFProperty<FProperty>(AActor::StaticClass(), "ActorLabel"));
+		PostEditChangeProperty(PropertyEvent);
 
-				FCoreDelegates::OnActorLabelChanged.Broadcast(this);
-			}
-		}
+		FCoreDelegates::OnActorLabelChanged.Broadcast(this);
+	}
+}
+
+void AActor::SetInitialActorLabel(FStringView InitialLabel)
+{
+	if (!ensureMsgf(bHasFinishedSpawning == false, TEXT("It is invalid to call AActor::SetInitialActorLabel after it has finished spawning.")))
+	{
+		return;
+	}
+
+	FText OutErrorMessage;
+	if (!FActorEditorUtils::ValidateActorLabel(InitialLabel, OutErrorMessage))
+	{
+		UE_LOG(LogActor, Warning, TEXT("SetActorLabel failed: %s"), *OutErrorMessage.ToString());
+	}
+	else
+	{
+		ActorLabel = FString(InitialLabel);
 	}
 }
 
diff --git a/Engine/Source/Runtime/Engine/Private/ActorEditorUtils.cpp b/Engine/Source/Runtime/Engine/Private/ActorEditorUtils.cpp
--- a/Engine/Source/Runtime/Engine/Private/ActorEditorUtils.cpp
+++ b/Engine/Source/Runtime/Engine/Private/ActorEditorUtils.cpp
@@ -122,25 +122,28 @@ namespace FActorEditorUtils
 
 	bool ValidateActorName(const FText& InName, FText& OutErrorMessage)
 	{
-		FText TrimmedLabel = FText::TrimPrecedingAndTrailing(InName);
+		return ValidateActorLabel(FStringView(InName.ToString()), OutErrorMessage);
+	}
+
+	bool ValidateActorLabel(FStringView InName, FText& OutErrorMessage)
+	{
+		InName.TrimStartAndEndInline();
 
-		if (TrimmedLabel.IsEmpty())
+		if (InName.Len() == 0)
 		{
-			OutErrorMessage = LOCTEXT("RenameFailed_LeftBlank", "Names cannot be left blank");
+			OutErrorMessage = LOCTEXT("ValidateActorNameString_Empty", "Actor names cannot be empty");
 			return false;
 		}
 
-		if (TrimmedLabel.ToString().Len() >= NAME_SIZE)
+		if (InName.Len() >= NAME_SIZE)
 		{
-			FFormatNamedArguments Arguments;
-			Arguments.Add(TEXT("CharCount"), NAME_SIZE);
-			OutErrorMessage = FText::Format(LOCTEXT("RenameFailed_TooLong", "Names must be less than {CharCount} characters long."), Arguments);
+			OutErrorMessage = FText::Format(LOCTEXT("ValidateActorNameString_TooLong", "Names must be less than {0} characters long."), NAME_SIZE);
 			return false;
 		}
 
-		if (FName(*TrimmedLabel.ToString()) == NAME_None)
+		if (InName.Compare("None", ESearchCase::IgnoreCase) == 0)
 		{
-			OutErrorMessage = LOCTEXT("RenameFailed_ReservedNameNone", "\"None\" is a reserved term and cannot be used for actor names");
+			OutErrorMessage = LOCTEXT("ValidateActorNameString_ReservedNameNone", "\"None\" is a reserved term and cannot be used for actor names");
 			return false;
 		}
 
diff --git a/Engine/Source/Runtime/Engine/Private/LevelActor.cpp b/Engine/Source/Runtime/Engine/Private/LevelActor.cpp
--- a/Engine/Source/Runtime/Engine/Private/LevelActor.cpp
+++ b/Engine/Source/Runtime/Engine/Private/LevelActor.cpp
@@ -678,7 +678,14 @@ AActor* UWorld::SpawnActor( UClass* Class, FTransform const* UserTransformPtr, c
 #endif
 
 #if WITH_EDITOR
-	Actor->ClearActorLabel(); // Clear label on newly spawned actors
+	if (SpawnParameters.InitialActorLabel.Len() != 0)
+	{
+		Actor->SetInitialActorLabel(SpawnParameters.InitialActorLabel);
+	}
+	else
+	{
+		Actor->ClearActorLabel(); // Clear label on newly spawned actors
+	}
 
 	// Generate the actor's Guid
 	FGuid ActorGuid;
diff --git a/Engine/Source/Runtime/Engine/Public/ActorEditorUtils.h b/Engine/Source/Runtime/Engine/Public/ActorEditorUtils.h
--- a/Engine/Source/Runtime/Engine/Public/ActorEditorUtils.h
+++ b/Engine/Source/Runtime/Engine/Public/ActorEditorUtils.h
@@ -67,4 +67,13 @@ namespace FActorEditorUtils
 	 */ 
 	ENGINE_API bool ValidateActorName(const FText& InName, FText& OutErrorMessage);
 
+	/**
+	 * Validates that a name is suitable for an actor
+	 *
+	 * @param InName			The proposed actor name
+	 * @param OutErrorMessage	Any error messages generated as a result of validation
+	 * @return true if the actor name is valid, false if not.
+	 */ 
+	ENGINE_API bool ValidateActorLabel(FStringView InName, FText& OutErrorMessage);
+
 };
diff --git a/Engine/Source/Runtime/LevelSequence/Private/LevelSequenceActorSpawner.cpp b/Engine/Source/Runtime/LevelSequence/Private/LevelSequenceActorSpawner.cpp
--- a/Engine/Source/Runtime/LevelSequence/Private/LevelSequenceActorSpawner.cpp
+++ b/Engine/Source/Runtime/LevelSequence/Private/LevelSequenceActorSpawner.cpp
@@ -140,10 +140,27 @@ UObject* FLevelSequenceActorSpawner::SpawnObject(FMovieSceneSpawnable& Spawnable
 		}
 	}
 
+	FStringView SpawnLabel;
+
+#if WITH_EDITOR
+	const IConsoleVariable* AllowSetActorLabelCvar = IConsoleManager::Get().FindConsoleVariable(TEXT("LevelSequence.EnableReadableActorLabelsForSpawnables"));
+	const bool bAllowSetActorLabel = AllowSetActorLabelCvar && AllowSetActorLabelCvar->GetBool();
+	
+	// Historically, setting the actor label has caused performance issues in some scenarios (by causing async loading flushes); however, there's no
+	// evidence for this anymore, so the cvar is here to turn off this behavior if needed.
+	if ((WorldContext->WorldType == EWorldType::Editor) || bAllowSetActorLabel)
+	{
+		SpawnLabel = Spawnable.GetName();
+	}
+#endif
+
 	// Spawn the puppet actor
 	FActorSpawnParameters SpawnInfo;
 	{
 		SpawnInfo.Name = SpawnName;
+#if WITH_EDITOR
+		SpawnInfo.InitialActorLabel = SpawnLabel;
+#endif
 		SpawnInfo.ObjectFlags = ObjectFlags;
 		SpawnInfo.SpawnCollisionHandlingOverride = ESpawnActorCollisionHandlingMethod::AlwaysSpawn;
 		// @todo: Spawning with a non-CDO template is fraught with issues
@@ -241,18 +258,6 @@ UObject* FLevelSequenceActorSpawner::SpawnObject(FMovieSceneSpawnable& Spawnable
 	const bool bIsDefaultTransform = true;
 	SpawnedActor->FinishSpawning(SpawnTransform, bIsDefaultTransform);
 
-#if WITH_EDITOR
-	const IConsoleVariable* AllowSetActorLabelCvar = IConsoleManager::Get().FindConsoleVariable(TEXT("LevelSequence.EnableReadableActorLabelsForSpawnables"));
-	const bool bAllowSetActorLabel = AllowSetActorLabelCvar && AllowSetActorLabelCvar->GetBool();
-	
-	// Historically, setting the actor label has caused performance issues in some scenarios (by causing async loading flushes); however, there's no
-	// evidence for this anymore, so the cvar is here to turn off this behavior if needed.
-	if ((WorldContext->WorldType == EWorldType::Editor) || bAllowSetActorLabel)
-	{
-		SpawnedActor->SetActorLabel(Spawnable.GetName());
-	}
-#endif
-
 	return SpawnedActor;
 }
 
diff --git a/Engine/Source/Runtime/MovieScene/Private/Bindings/MovieSceneSpawnableBinding.cpp b/Engine/Source/Runtime/MovieScene/Private/Bindings/MovieSceneSpawnableBinding.cpp
--- a/Engine/Source/Runtime/MovieScene/Private/Bindings/MovieSceneSpawnableBinding.cpp
+++ b/Engine/Source/Runtime/MovieScene/Private/Bindings/MovieSceneSpawnableBinding.cpp
@@ -93,37 +93,6 @@ UObject* UMovieSceneSpawnableBindingBase::SpawnObject(const FGuid& BindingId, in
 			}
 		}
 #endif
-
-#if WITH_EDITOR
-		// Historically, setting the actor label has caused performance issues in some scenarios (by causing async loading flushes); however, there's no
-		// evidence for this anymore, so the cvar is here to turn off this behavior if needed.
-		if ((WorldContext->WorldType == EWorldType::Editor) || CVarEnableReadableActorLabelsForSpawnables->GetBool())
-		{
-			FString BindingName = GetDesiredBindingName();
-			FString ActorLabel = !BindingName.IsEmpty() ? BindingName : SpawnName.ToString();
-
-			if (FMovieScenePossessable* Possessable = MovieScene.FindPossessable(BindingId))
-			{
-				ActorLabel = Possessable->GetName();
-
-				if (UMovieSceneSequence* Sequence = MovieScene.GetTypedOuter<UMovieSceneSequence>())
-				{
-					if (const FMovieSceneBindingReferences* BindingReferences = Sequence->GetBindingReferences())
-					{
-						if (BindingReferences->GetReferences(BindingId).Num() > 1)
-						{
-							// If there are multiple bound objects, use the Object Template actor label instead of the possessable name
-							if (AActor* ActorTemplate = Cast<AActor>(GetObjectTemplate()))
-							{
-								ActorLabel = ActorTemplate->GetActorLabel();
-							}
-						}
-					}
-				}
-			}
-			SpawnedActor->SetActorLabel(ActorLabel);
-		}
-#endif
 	}
 
 	// Allows derived classes to perform post-spawn logic such as mesh setup on actors.
diff --git a/Engine/Source/Runtime/MovieSceneTracks/Private/Bindings/MovieSceneSpawnableActorBinding.cpp b/Engine/Source/Runtime/MovieSceneTracks/Private/Bindings/MovieSceneSpawnableActorBinding.cpp
--- a/Engine/Source/Runtime/MovieSceneTracks/Private/Bindings/MovieSceneSpawnableActorBinding.cpp
+++ b/Engine/Source/Runtime/MovieSceneTracks/Private/Bindings/MovieSceneSpawnableActorBinding.cpp
@@ -39,11 +39,50 @@ UObject* UMovieSceneSpawnableActorBindingBase::SpawnObjectInternal(UWorld* World
 
 	AActor* ActorTemplate = GetActorTemplate();
 
+	FString SpawnLabel;
+
+#if WITH_EDITOR
+	const IConsoleVariable* AllowSetActorLabelCvar = IConsoleManager::Get().FindConsoleVariable(TEXT("LevelSequence.EnableReadableActorLabelsForSpawnables"));
+	const bool bAllowSetActorLabel = AllowSetActorLabelCvar && AllowSetActorLabelCvar->GetBool();
+	
+	// Historically, setting the actor label has caused performance issues in some scenarios (by causing async loading flushes); however, there's no
+	// evidence for this anymore, so the cvar is here to turn off this behavior if needed.
+	if ((WorldContext->WorldType == EWorldType::Editor) || bAllowSetActorLabel)
+	{
+		if (FMovieScenePossessable* Possessable = MovieScene.FindPossessable(BindingId))
+		{
+			UMovieSceneSequence*                Sequence          = MovieScene.GetTypedOuter<UMovieSceneSequence>();
+			const FMovieSceneBindingReferences* BindingReferences = Sequence ? Sequence->GetBindingReferences() : nullptr;
+
+			if (ActorTemplate && BindingReferences && BindingReferences->GetReferences(BindingId).Num() > 1)
+			{
+				// If there are multiple bound objects, use the Object Template actor label instead of the possessable name
+				SpawnLabel = ActorTemplate->GetActorLabel();
+			}
+			else
+			{
+				SpawnLabel = Possessable->GetName();
+			}
+		}
+		else
+		{
+			SpawnLabel = GetDesiredBindingName();
+			if (SpawnLabel.Len() == 0)
+			{
+				SpawnLabel = SpawnName.ToString();
+			}
+		}
+	}
+#endif
+
 	// Spawn the actor
 	FActorSpawnParameters SpawnInfo;
 	{
 		SpawnInfo.Name = SpawnName;
 		SpawnInfo.ObjectFlags = SpawnObjectFlags;
+#if WITH_EDITOR
+		SpawnInfo.InitialActorLabel = SpawnLabel;
+#endif
 		SpawnInfo.SpawnCollisionHandlingOverride = ESpawnActorCollisionHandlingMethod::AlwaysSpawn;
 		// allow pre-construction variables to be set.
 		SpawnInfo.bDeferConstruction = true;
