Sequencer: Fix further issue with postloading level sequences with dynamic bindings during cook. Under some circumstances the blueprint compilation can interfere with other compilations happening during load.

diff --git a/Engine/Source/Runtime/LevelSequence/Private/LevelSequence.cpp b/Engine/Source/Runtime/LevelSequence/Private/LevelSequence.cpp
--- a/Engine/Source/Runtime/LevelSequence/Private/LevelSequence.cpp
+++ b/Engine/Source/Runtime/LevelSequence/Private/LevelSequence.cpp
@@ -432,7 +432,14 @@ void ULevelSequence::PostLoad()
 		}
 		if (bConvertedDynamicBinding && FixupDynamicBindingsEvent.IsBound())
 		{
-			FixupDynamicBindingsEvent.Broadcast(this);
+			if (!GetPackage()->GetHasBeenEndLoaded())
+			{
+				FCoreUObjectDelegates::OnEndLoadPackage.AddUObject(this, &ULevelSequence::OnEndLoadPackage);
+			}
+			else
+			{
+				FixupDynamicBindingsEvent.Broadcast(this);
+			}
 		}
 	}
 
@@ -440,6 +447,16 @@ void ULevelSequence::PostLoad()
 }
 
 #if WITH_EDITOR
+void ULevelSequence::OnEndLoadPackage(const FEndLoadPackageContext& Context)
+{
+	FCoreUObjectDelegates::OnEndLoadPackage.RemoveAll(this);
+
+	if (FixupDynamicBindingsEvent.IsBound())
+	{
+		FixupDynamicBindingsEvent.Broadcast(this);
+	}
+}
+
 
 bool ULevelSequence::ConvertOldSpawnables()
 {
diff --git a/Engine/Source/Runtime/LevelSequence/Public/LevelSequence.h b/Engine/Source/Runtime/LevelSequence/Public/LevelSequence.h
--- a/Engine/Source/Runtime/LevelSequence/Public/LevelSequence.h
+++ b/Engine/Source/Runtime/LevelSequence/Public/LevelSequence.h
@@ -139,6 +139,7 @@ private:
 	bool ConvertOldSpawnables();
 	// Used to convert possessables using deprecated 'DynamicBinding' to a custom director blueprint binding
 	void ConvertDynamicBindingPossessable(FMovieScenePossessable& Possessable);
+	void OnEndLoadPackage(const FEndLoadPackageContext& Context);
 #endif
 
 #endif // WITH_EDITOR
