Anim in Engine: Fix linked anim level equality check

diff --git a/Engine/Source/Runtime/LevelSequence/Public/LevelSequenceAnimSequenceLink.h b/Engine/Source/Runtime/LevelSequence/Public/LevelSequenceAnimSequenceLink.h
--- a/Engine/Source/Runtime/LevelSequence/Public/LevelSequenceAnimSequenceLink.h
+++ b/Engine/Source/Runtime/LevelSequence/Public/LevelSequenceAnimSequenceLink.h
@@ -94,14 +94,13 @@ public:
 				if (bUseCustomTimeRange == false || (InCustomStartFrame == CustomStartFrame &&
 					InCustomEndFrame == CustomEndFrame && InCustomDisplayRate == CustomDisplayRate))
 				{
-					return true;
-				}
-			}
-			if (bUseCustomFrameRate == bInUseCustomFrameRate)
-			{
-				if (bUseCustomFrameRate == false || (InCustomFrameRate == CustomFrameRate))
-				{
-					return true;
+					if (bUseCustomFrameRate == bInUseCustomFrameRate)
+					{
+						if (bUseCustomFrameRate == false || (InCustomFrameRate == CustomFrameRate))
+						{
+							return true;
+						}
+					}
 				}
 			}
 		}
