TakeRecorder / MovieScene: Move the newly added IsSubSequence asset registry tag from UTakeMetaData to UMovieSceneShotMetaData.
Add IsRecorded to UMovieSceneShotMetaData.
Take Recorder will now create UMovieSceneShotMetaData meta data when appropriate to record these values.
Change the recently added SubsequenceDirectory in Take Recorder to match the case of other uses of SubSequence.

diff --git a/Engine/Plugins/VirtualProduction/Takes/Source/TakeRecorder/Private/Recorder/TakeRecorder.cpp b/Engine/Plugins/VirtualProduction/Takes/Source/TakeRecorder/Private/Recorder/TakeRecorder.cpp
--- a/Engine/Plugins/VirtualProduction/Takes/Source/TakeRecorder/Private/Recorder/TakeRecorder.cpp
+++ b/Engine/Plugins/VirtualProduction/Takes/Source/TakeRecorder/Private/Recorder/TakeRecorder.cpp
@@ -14,6 +14,7 @@
 #include "AssetRegistry/IAssetRegistry.h"
 #include "ISequencer.h"
 #include "LevelSequence.h"
+#include "MetaData/MovieSceneShotMetaData.h"
 #include "Misc/CommandLine.h"
 #include "Misc/Parse.h"
 #include "MovieSceneTimeHelpers.h"
@@ -878,7 +879,7 @@ void UTakeRecorder::PreRecord()
 	TakeRecorderSourcesSettings.bSaveRecordedAssets = (!Parameters.bDisableRecordingAndSave && Parameters.User.bSaveRecordedAssets) || GEditor == nullptr;
 	TakeRecorderSourcesSettings.bRemoveRedundantTracks = Parameters.User.bRemoveRedundantTracks;
 	TakeRecorderSourcesSettings.bAutoLock = Parameters.User.bAutoLock;
-	TakeRecorderSourcesSettings.SubsequenceDirectory = Parameters.Project.SubsequenceDirectory;
+	TakeRecorderSourcesSettings.SubSequenceDirectory = Parameters.Project.SubSequenceDirectory;
 
 	TSharedPtr<ISequencer> Sequencer = WeakSequencer.Pin();
 	if (Sequencer)
@@ -1170,6 +1171,10 @@ void UTakeRecorder::StopInternal(const bool bCancelled)
 				SequenceAsset->GetMovieScene()->SetReadOnly(true);
 			}
 
+			// We have recorded this sequence, mark it.
+			UMovieSceneShotMetaData* MovieSceneMetaData = SequenceAsset->FindOrAddMetaData<UMovieSceneShotMetaData>();
+			MovieSceneMetaData->SetIsRecorded(true);
+
 			UTakeMetaData* AssetMetaData = SequenceAsset->FindMetaData<UTakeMetaData>();
 			check(AssetMetaData);
 
diff --git a/Engine/Plugins/VirtualProduction/Takes/Source/TakeRecorder/Private/Recorder/TakeRecorderParameters.cpp b/Engine/Plugins/VirtualProduction/Takes/Source/TakeRecorder/Private/Recorder/TakeRecorderParameters.cpp
--- a/Engine/Plugins/VirtualProduction/Takes/Source/TakeRecorder/Private/Recorder/TakeRecorderParameters.cpp
+++ b/Engine/Plugins/VirtualProduction/Takes/Source/TakeRecorder/Private/Recorder/TakeRecorderParameters.cpp
@@ -24,7 +24,7 @@ FTakeRecorderUserParameters::FTakeRecorderUserParameters()
 }
 
 FTakeRecorderProjectParameters::FTakeRecorderProjectParameters()
-	: SubsequenceDirectory(TEXT("{takeName}_Subscenes"))
+	: SubSequenceDirectory(TEXT("{takeName}_Subscenes"))
 	, RecordingClockSource(EUpdateClockSource::RelativeTimecode)
 	, bStartAtCurrentTimecode(true)
 	, bRecordTimecode(false)
diff --git a/Engine/Plugins/VirtualProduction/Takes/Source/TakeRecorder/Public/Recorder/TakeRecorderParameters.h b/Engine/Plugins/VirtualProduction/Takes/Source/TakeRecorder/Public/Recorder/TakeRecorderParameters.h
--- a/Engine/Plugins/VirtualProduction/Takes/Source/TakeRecorder/Public/Recorder/TakeRecorderParameters.h
+++ b/Engine/Plugins/VirtualProduction/Takes/Source/TakeRecorder/Public/Recorder/TakeRecorderParameters.h
@@ -90,7 +90,7 @@ struct FTakeRecorderProjectParameters
 	 * The directory to record subscenes into.
 	 */
 	UPROPERTY(config, EditAnywhere, BlueprintReadWrite, Category = "Take Recorder", meta = (NamingTokens))
-	FString SubsequenceDirectory;
+	FString SubSequenceDirectory;
 	
 	/**
 	 * The default name to use for the Slate information
diff --git a/Engine/Plugins/VirtualProduction/Takes/Source/TakesCore/Private/TakeMetaData.cpp b/Engine/Plugins/VirtualProduction/Takes/Source/TakesCore/Private/TakeMetaData.cpp
--- a/Engine/Plugins/VirtualProduction/Takes/Source/TakesCore/Private/TakeMetaData.cpp
+++ b/Engine/Plugins/VirtualProduction/Takes/Source/TakesCore/Private/TakeMetaData.cpp
@@ -25,7 +25,6 @@ const FName UTakeMetaData::AssetRegistryTag_TimecodeIn = "TakeMetaData_TimecodeI
 const FName UTakeMetaData::AssetRegistryTag_TimecodeOut = "TakeMetaData_TimecodeOut";
 const FName UTakeMetaData::AssetRegistryTag_Description = "TakeMetaData_Description";
 const FName UTakeMetaData::AssetRegistryTag_LevelPath   = "TakeMetaData_LevelPath";
-const FName UTakeMetaData::AssetRegistryTag_IsSubSequence = "TakeMetaData_IsSubSequence";
 
 UTakeMetaData::UTakeMetaData(const FObjectInitializer& ObjInit)
 	: Super(ObjInit)
@@ -346,19 +345,6 @@ void UTakeMetaData::SetFrameRateFromTimecode(bool InFromTimecode)
 	}
 }
 
-bool UTakeMetaData::IsSubSequence() const
-{
-	return bIsSubSequence;
-}
-
-void UTakeMetaData::SetIsSubSequence(bool bNewValue)
-{
-	if (!bIsLocked)
-	{
-		bIsSubSequence = bNewValue;
-	}
-}
-
 UTakeMetaData::FOnTakeSlateChanged& UTakeMetaData::OnTakeSlateChanged()
 {
 	static FOnTakeSlateChanged OnSlateChanged;
@@ -387,8 +373,6 @@ void UTakeMetaData::ExtendAssetRegistryTags(FAssetRegistryTagsContext Context) c
 										FAssetRegistryTag::ETagType::TT_Numerical,		FAssetRegistryTag::TD_None));
 	Context.AddTag(FAssetRegistryTag(	AssetRegistryTag_Description,					Description,
 										FAssetRegistryTag::ETagType::TT_Alphabetical,	FAssetRegistryTag::TD_None));
-	Context.AddTag(FAssetRegistryTag(	AssetRegistryTag_IsSubSequence,					IsSubSequence() ? TEXT("True") : TEXT("False"),
-										FAssetRegistryTag::ETagType::TT_Alphabetical,	FAssetRegistryTag::TD_None));
 }
 
 void UTakeMetaData::ExtendAssetRegistryTagMetaData(TMap<FName, FAssetRegistryTagMetadata>& OutMetadata) const
@@ -427,11 +411,6 @@ void UTakeMetaData::ExtendAssetRegistryTagMetaData(TMap<FName, FAssetRegistryTag
 		.SetDisplayName(NSLOCTEXT("TakeMetaData", "LevelPath_Label", "Map"))
 		.SetTooltip(    NSLOCTEXT("TakeMetaData", "LevelPath_Tip",   "Map used for this take"))
 	);
-
-	OutMetadata.Add(AssetRegistryTag_IsSubSequence, FAssetRegistryTagMetadata()
-		.SetDisplayName(NSLOCTEXT("TakeMetaData", "IsSubSequence_Label", "Is SubSequence"))
-		.SetTooltip(    NSLOCTEXT("TakeMetaData", "IsSubSequence_Tip",   "If this take was recorded as a subsequence"))
-	);
 }
 
 
diff --git a/Engine/Plugins/VirtualProduction/Takes/Source/TakesCore/Private/TakeRecorderSources.cpp b/Engine/Plugins/VirtualProduction/Takes/Source/TakesCore/Private/TakeRecorderSources.cpp
--- a/Engine/Plugins/VirtualProduction/Takes/Source/TakesCore/Private/TakeRecorderSources.cpp
+++ b/Engine/Plugins/VirtualProduction/Takes/Source/TakesCore/Private/TakeRecorderSources.cpp
@@ -20,6 +20,7 @@
 #include "Engine/TimecodeProvider.h"
 #include "Misc/App.h"
 #include "Editor.h"
+#include "MetaData/MovieSceneShotMetaData.h"
 #include "NamingTokens/TakeRecorderNamingTokensContext.h"
 #include "UObject/Package.h"
 
@@ -476,12 +477,17 @@ void UTakeRecorderSources::StopRecording(class ULevelSequence* InSequence, const
 					SubSequence->GetMovieScene()->SetReadOnly(true);
 				}
 
-				// Lock the meta data so it can't be changed without implicitly unlocking it now
 				ULevelSequence* SequenceAsset = CastChecked<ULevelSequence>(SubSequence);
-				UTakeMetaData* AssetMetaData = SequenceAsset->FindMetaData<UTakeMetaData>();
-				check(AssetMetaData);
-				AssetMetaData->SetIsSubSequence(true);
-				AssetMetaData->Lock();
+
+				// Mark that this subsequence has been recorded and is a subsequence.
+				UMovieSceneShotMetaData* MovieSequenceMetaData = SequenceAsset->FindOrAddMetaData<UMovieSceneShotMetaData>();
+				MovieSequenceMetaData->SetIsRecorded(true);
+				MovieSequenceMetaData->SetIsSubSequence(true);
+
+				// Lock the meta data so it can't be changed without implicitly unlocking it now
+				UTakeMetaData* TakeMetaData = SequenceAsset->FindMetaData<UTakeMetaData>();
+				check(TakeMetaData);
+				TakeMetaData->Lock();
 
 				SubSection->SetRange(SubSequence->GetMovieScene()->GetPlaybackRange());
 
@@ -577,23 +583,23 @@ ULevelSequence* UTakeRecorderSources::CreateSubSequenceForSource(ULevelSequence*
 		}
 	}
 	
-	FString SubsequenceDirectory;
+	FString SubSequenceDirectory;
 	if (const UTakeMetaData* AssetMetaData = InRootSequence->FindMetaData<UTakeMetaData>())
 	{
 		// User specified subsequence directory.
 		UTakeRecorderNamingTokensContext* Context = NewObject<UTakeRecorderNamingTokensContext>();
 		Context->TakeMetaData = AssetMetaData;
-		SubsequenceDirectory = InSettings.SubsequenceDirectory;
-		SubsequenceDirectory = AssetMetaData->GenerateAssetPath(SubsequenceDirectory, Context);
+		SubSequenceDirectory = InSettings.SubSequenceDirectory;
+		SubSequenceDirectory = AssetMetaData->GenerateAssetPath(SubSequenceDirectory, Context);
 	}
 	
-	if (SubsequenceDirectory.IsEmpty())
+	if (SubSequenceDirectory.IsEmpty())
 	{
 		// Fallback to original behavior.
-		SubsequenceDirectory = FString::Printf(TEXT("%s_Subscenes"), *SequenceName);
+		SubSequenceDirectory = FString::Printf(TEXT("%s_Subscenes"), *SequenceName);
 	}
 	
-	FString NewPath = FString::Printf(TEXT("%s/%s/%s"), *SequenceDirectory, *SubsequenceDirectory, *SubSequenceAssetName);
+	FString NewPath = FString::Printf(TEXT("%s/%s/%s"), *SequenceDirectory, *SubSequenceDirectory, *SubSequenceAssetName);
 
 	ULevelSequence* OutAsset = nullptr;
 	TakesUtils::CreateNewAssetPackage<ULevelSequence>(NewPath, OutAsset, nullptr, ExistingSubSequence, InRootSequence->GetClass());
diff --git a/Engine/Plugins/VirtualProduction/Takes/Source/TakesCore/Public/TakeMetaData.h b/Engine/Plugins/VirtualProduction/Takes/Source/TakesCore/Public/TakeMetaData.h
--- a/Engine/Plugins/VirtualProduction/Takes/Source/TakesCore/Public/TakeMetaData.h
+++ b/Engine/Plugins/VirtualProduction/Takes/Source/TakesCore/Public/TakeMetaData.h
@@ -48,9 +48,6 @@ public:
 	/** The asset registry tag that contains the level-path for this meta-data */
 	static const FName AssetRegistryTag_LevelPath;
 	
-	/** The asset registry tag that contains whether this is a subsequence. */
-	static const FName AssetRegistryTag_IsSubSequence;
-	
 	/**
 	 * Access the global config instance that houses default settings for take meta data for a given project
 	 */
@@ -281,16 +278,6 @@ public:
 	UFUNCTION(BlueprintCallable, Category = "Take")
 	void SetFrameRateFromTimecode(bool InFromTimecode);
 
-	/**
-	 * True if this take meta data was recorded as part of a subsequence.
-	 */
-	UFUNCTION(BlueprintCallable, Category = "Take")
-	bool IsSubSequence() const;
-
-	/** Set whether the take meta data is part of a subsequence. */
-	UFUNCTION(BlueprintCallable, Category = "Take")
-	void SetIsSubSequence(bool bNewValue);
-
 	DECLARE_MULTICAST_DELEGATE_TwoParams(FOnTakeSlateChanged, const FString& SlateName, UTakeMetaData* TakeMetaData);
 	/** Multicast delegate emitted when the slate is changed. */
 	static FOnTakeSlateChanged& OnTakeSlateChanged();
@@ -350,8 +337,4 @@ private:
 	/** Whether or not we get or frame rate from Timecode, default to true */
 	UPROPERTY()
 	bool bFrameRateFromTimecode;
-
-	/** If this meta data was recorded as a subsequence. */
-	UPROPERTY()
-	bool bIsSubSequence = false;
 };
diff --git a/Engine/Plugins/VirtualProduction/Takes/Source/TakesCore/Public/TakeRecorderSources.h b/Engine/Plugins/VirtualProduction/Takes/Source/TakesCore/Public/TakeRecorderSources.h
--- a/Engine/Plugins/VirtualProduction/Takes/Source/TakesCore/Public/TakeRecorderSources.h
+++ b/Engine/Plugins/VirtualProduction/Takes/Source/TakesCore/Public/TakeRecorderSources.h
@@ -25,7 +25,7 @@ struct FTakeRecorderSourcesSettings
 	bool bAutoLock;
 	bool bRemoveRedundantTracks;
 	FTransform TransformOrigin;
-	FString SubsequenceDirectory;
+	FString SubSequenceDirectory;
 };
 
 /**
diff --git a/Engine/Plugins/VirtualProduction/VirtualCamera/Source/VirtualCamera/Private/FunctionLibraries/TakeMetaDataTagsFunctionLibrary.cpp b/Engine/Plugins/VirtualProduction/VirtualCamera/Source/VirtualCamera/Private/FunctionLibraries/TakeMetaDataTagsFunctionLibrary.cpp
--- a/Engine/Plugins/VirtualProduction/VirtualCamera/Source/VirtualCamera/Private/FunctionLibraries/TakeMetaDataTagsFunctionLibrary.cpp
+++ b/Engine/Plugins/VirtualProduction/VirtualCamera/Source/VirtualCamera/Private/FunctionLibraries/TakeMetaDataTagsFunctionLibrary.cpp
@@ -69,15 +69,6 @@ FName UTakeMetaDataTagsFunctionLibrary::GetTakeMetaDataTag_LevelPath()
 #endif
 }
 
-FName UTakeMetaDataTagsFunctionLibrary::GetTakeMetaDataTag_IsSubSequence()
-{
-#if WITH_EDITOR
-	return UTakeMetaData::AssetRegistryTag_IsSubSequence;
-#else
-	return "TakeMetaData_IsSubSequence";
-#endif
-}
-
 TSet<FName> UTakeMetaDataTagsFunctionLibrary::GetAllTakeMetaDataTags()
 {
 	return {
@@ -87,7 +78,6 @@ TSet<FName> UTakeMetaDataTagsFunctionLibrary::GetAllTakeMetaDataTags()
 		GetTakeMetaDataTag_TimecodeIn(),
 		GetTakeMetaDataTag_TimecodeOut(),
 		GetTakeMetaDataTag_Description(),
-		GetTakeMetaDataTag_LevelPath(),
-		GetTakeMetaDataTag_IsSubSequence()
+		GetTakeMetaDataTag_LevelPath()
 	};
 }
diff --git a/Engine/Plugins/VirtualProduction/VirtualCamera/Source/VirtualCamera/Public/FunctionLibraries/TakeMetaDataTagsFunctionLibrary.h b/Engine/Plugins/VirtualProduction/VirtualCamera/Source/VirtualCamera/Public/FunctionLibraries/TakeMetaDataTagsFunctionLibrary.h
--- a/Engine/Plugins/VirtualProduction/VirtualCamera/Source/VirtualCamera/Public/FunctionLibraries/TakeMetaDataTagsFunctionLibrary.h
+++ b/Engine/Plugins/VirtualProduction/VirtualCamera/Source/VirtualCamera/Public/FunctionLibraries/TakeMetaDataTagsFunctionLibrary.h
@@ -40,10 +40,6 @@ public:
 	UFUNCTION(BlueprintPure, Category = "VirtualCamera")
 	static FName GetTakeMetaDataTag_LevelPath();
 
-	/** The asset registry tag that specifies if the meta-data is a subsequence. */
-	UFUNCTION(BlueprintPure, Category = "VirtualCamera")
-	static FName GetTakeMetaDataTag_IsSubSequence();
-
 	/** Gets all tags for TakeMetaData */
 	UFUNCTION(BlueprintPure, Category = "VirtualCamera")
 	static TSet<FName> GetAllTakeMetaDataTags();
diff --git a/Engine/Source/Runtime/LevelSequence/Private/LevelSequenceShotMetaDataLibrary.cpp b/Engine/Source/Runtime/LevelSequence/Private/LevelSequenceShotMetaDataLibrary.cpp
--- a/Engine/Source/Runtime/LevelSequence/Private/LevelSequenceShotMetaDataLibrary.cpp
+++ b/Engine/Source/Runtime/LevelSequence/Private/LevelSequenceShotMetaDataLibrary.cpp
@@ -51,6 +51,30 @@ bool ULevelSequenceShotMetaDataLibrary::GetIsFlagged(const ULevelSequence* InLev
 	return true;
 }
 
+bool ULevelSequenceShotMetaDataLibrary::GetIsRecorded(const ULevelSequence* InLevelSequence, bool& bOutIsRecorded)
+{
+	const UMovieSceneShotMetaData* MetaData = UE::Sequencer::TakesMetaDataDetail::FindMetaData(InLevelSequence);
+	if (!MetaData || !MetaData->GetIsRecorded().IsSet())
+	{
+		return false;
+	}
+	
+	bOutIsRecorded = *MetaData->GetIsRecorded();
+	return true;
+}
+
+bool ULevelSequenceShotMetaDataLibrary::GetIsSubSequence(const ULevelSequence* InLevelSequence, bool& bOutIsSubSequence)
+{
+	const UMovieSceneShotMetaData* MetaData = UE::Sequencer::TakesMetaDataDetail::FindMetaData(InLevelSequence);
+	if (!MetaData || !MetaData->GetIsSubSequence().IsSet())
+	{
+		return false;
+	}
+	
+	bOutIsSubSequence = *MetaData->GetIsSubSequence();
+	return true;
+}
+
 bool ULevelSequenceShotMetaDataLibrary::GetFavoriteRating(const ULevelSequence* InLevelSequence, int32& OutFavoriteRating)
 {
 	UMovieSceneShotMetaData* MetaData = UE::Sequencer::TakesMetaDataDetail::FindMetaData(InLevelSequence);
@@ -75,6 +99,18 @@ bool ULevelSequenceShotMetaDataLibrary::HasIsFlagged(const ULevelSequence* InLev
 	return GetIsFlagged(InLevelSequence, Dummy);
 }
 
+bool ULevelSequenceShotMetaDataLibrary::HasIsRecorded(const ULevelSequence* InLevelSequence)
+{
+	bool Dummy = false;
+	return GetIsRecorded(InLevelSequence, Dummy);
+}
+
+bool ULevelSequenceShotMetaDataLibrary::HasIsSubSequence(const ULevelSequence* InLevelSequence)
+{
+	bool Dummy = false;
+	return GetIsSubSequence(InLevelSequence, Dummy);
+}
+
 bool ULevelSequenceShotMetaDataLibrary::HasFavoriteRating(const ULevelSequence* InLevelSequence)
 {
 	int32 Dummy = 0;
@@ -97,6 +133,22 @@ void ULevelSequenceShotMetaDataLibrary::SetIsFlagged(ULevelSequence* InLevelSequ
 	}
 }
 
+void ULevelSequenceShotMetaDataLibrary::SetIsRecorded(ULevelSequence* InLevelSequence, bool bInIsRecorded)
+{
+	if (UMovieSceneShotMetaData* MetaData = UE::Sequencer::TakesMetaDataDetail::FindOrAddMetaData(InLevelSequence))
+	{
+		MetaData->SetIsRecorded(bInIsRecorded);
+	}
+}
+
+void ULevelSequenceShotMetaDataLibrary::SetIsSubSequence(ULevelSequence* InLevelSequence, bool bInIsSubSequence)
+{
+	if (UMovieSceneShotMetaData* MetaData = UE::Sequencer::TakesMetaDataDetail::FindOrAddMetaData(InLevelSequence))
+	{
+		MetaData->SetIsSubSequence(bInIsSubSequence);
+	}
+}
+
 void ULevelSequenceShotMetaDataLibrary::SetFavoriteRating(ULevelSequence* InLevelSequence, int32 InFavoriteRating)
 {
 	if (UMovieSceneShotMetaData* MetaData = UE::Sequencer::TakesMetaDataDetail::FindOrAddMetaData(InLevelSequence))
@@ -121,6 +173,22 @@ void ULevelSequenceShotMetaDataLibrary::ClearIsFlagged(ULevelSequence* InLevelSe
 	}
 }
 
+void ULevelSequenceShotMetaDataLibrary::ClearIsRecorded(ULevelSequence* InLevelSequence)
+{
+	if (UMovieSceneShotMetaData* MetaData = UE::Sequencer::TakesMetaDataDetail::FindOrAddMetaData(InLevelSequence))
+	{
+		MetaData->ClearIsRecorded();
+	}
+}
+
+void ULevelSequenceShotMetaDataLibrary::ClearIsSubSequence(ULevelSequence* InLevelSequence)
+{
+	if (UMovieSceneShotMetaData* MetaData = UE::Sequencer::TakesMetaDataDetail::FindOrAddMetaData(InLevelSequence))
+	{
+		MetaData->ClearIsSubSequence();
+	}
+}
+
 void ULevelSequenceShotMetaDataLibrary::ClearFavoriteRating(ULevelSequence* InLevelSequence)
 {
 	if (UMovieSceneShotMetaData* MetaData = UE::Sequencer::TakesMetaDataDetail::FindOrAddMetaData(InLevelSequence))
@@ -139,6 +207,16 @@ FName ULevelSequenceShotMetaDataLibrary::GetIsFlaggedAssetTag()
 	return UMovieSceneShotMetaData::AssetRegistryTag_bIsFlagged;
 }
 
+FName ULevelSequenceShotMetaDataLibrary::GetIsRecordedAssetTag()
+{
+	return UMovieSceneShotMetaData::AssetRegistryTag_bIsRecorded;
+}
+
+FName ULevelSequenceShotMetaDataLibrary::GetIsSubSequenceAssetTag()
+{
+	return UMovieSceneShotMetaData::AssetRegistryTag_bIsSubSequence;
+}
+
 FName ULevelSequenceShotMetaDataLibrary::GetFavoriteRatingAssetTag()
 {
 	return UMovieSceneShotMetaData::AssetRegistryTag_FavoriteRating;
@@ -154,6 +232,16 @@ bool ULevelSequenceShotMetaDataLibrary::GetIsFlaggedByAssetData(const FAssetData
 	return UMovieSceneShotMetaData::GetIsFlaggedByAssetData(InAssetData, bOutIsFlagged);
 }
 
+bool ULevelSequenceShotMetaDataLibrary::GetIsRecordedByAssetData(const FAssetData& InAssetData, bool& bOutIsRecorded)
+{
+	return UMovieSceneShotMetaData::GetIsRecordedByAssetData(InAssetData, bOutIsRecorded);
+}
+
+bool ULevelSequenceShotMetaDataLibrary::GetIsSubSequenceByAssetData(const FAssetData& InAssetData, bool& bOutIsSubSequence)
+{
+	return UMovieSceneShotMetaData::GetIsSubSequenceByAssetData(InAssetData, bOutIsSubSequence);
+}
+
 bool ULevelSequenceShotMetaDataLibrary::GetFavoriteRatingByAssetData(const FAssetData& InAssetData, int32& OutFavoriteRating)
 {
 	return UMovieSceneShotMetaData::GetFavoriteRatingByAssetData(InAssetData, OutFavoriteRating);
@@ -171,6 +259,18 @@ bool ULevelSequenceShotMetaDataLibrary::HasIsFlaggedByAssetData(const FAssetData
 	return GetIsFlaggedByAssetData(InAssetData, bDummy);
 }
 
+bool ULevelSequenceShotMetaDataLibrary::HasIsRecordedByAssetData(const FAssetData& InAssetData)
+{
+	bool bDummy = false;
+	return GetIsRecordedByAssetData(InAssetData, bDummy);
+}
+
+bool ULevelSequenceShotMetaDataLibrary::HasIsSubSequenceByAssetData(const FAssetData& InAssetData)
+{
+	bool bDummy = false;
+	return GetIsSubSequenceByAssetData(InAssetData, bDummy);
+}
+
 bool ULevelSequenceShotMetaDataLibrary::HasFavoriteRatingByAssetData(const FAssetData& InAssetData)
 {
 	int32 Dummy = 0;
diff --git a/Engine/Source/Runtime/LevelSequence/Public/LevelSequenceShotMetaDataLibrary.h b/Engine/Source/Runtime/LevelSequence/Public/LevelSequenceShotMetaDataLibrary.h
--- a/Engine/Source/Runtime/LevelSequence/Public/LevelSequenceShotMetaDataLibrary.h
+++ b/Engine/Source/Runtime/LevelSequence/Public/LevelSequenceShotMetaDataLibrary.h
@@ -33,6 +33,18 @@ public:
 	 */
 	UFUNCTION(BlueprintPure, Category = "Level Sequence")
 	static UE_API bool GetIsFlagged(const ULevelSequence* InLevelSequence, bool& bOutIsFlagged);
+	/**
+	 * Gets whether this was recorded.
+	 * @return Whether the value was set.
+	 */
+	UFUNCTION(BlueprintPure, Category = "Level Sequence")
+	static UE_API bool GetIsRecorded(const ULevelSequence* InLevelSequence, bool& bOutIsRecorded);
+	/**
+	 * Gets whether this was recorded as a subsequence.
+	 * @return Whether the value was set.
+	 */
+	UFUNCTION(BlueprintPure, Category = "Level Sequence")
+	static UE_API bool GetIsSubSequence(const ULevelSequence* InLevelSequence, bool& bOutIsSubSequence);
 	/**
 	 * Gets the favorite rating. The favorite rating is like a star rating, usually 1-3 if it was rated.
 	 * @return Whether the value was set by the user.
@@ -46,6 +58,12 @@ public:
 	/** @return Whether the user has ever set a value for this property. */
 	UFUNCTION(BlueprintPure, Category = "Level Sequence")
 	static UE_API bool HasIsFlagged(const ULevelSequence* InLevelSequence);
+	/** @return Whether IsRecorded has been set. */
+	UFUNCTION(BlueprintPure, Category = "Level Sequence")
+	static UE_API bool HasIsRecorded(const ULevelSequence* InLevelSequence);
+	/** @return Whether IsSubSequence has been set. */
+	UFUNCTION(BlueprintPure, Category = "Level Sequence")
+	static UE_API bool HasIsSubSequence(const ULevelSequence* InLevelSequence);
 	/** @return Whether the user has ever set a value for this property. */
 	UFUNCTION(BlueprintPure, Category = "Level Sequence")
 	static UE_API bool HasFavoriteRating(const ULevelSequence* InLevelSequence);
@@ -56,6 +74,12 @@ public:
 	/** Sets whether this was flagged by a user. */
 	UFUNCTION(BlueprintCallable, Category = "Level Sequence")
 	static UE_API void SetIsFlagged(ULevelSequence* InLevelSequence, bool bInIsFlagged);
+	/** Sets whether this sequence was recorded. */
+	UFUNCTION(BlueprintCallable, Category = "Level Sequence")
+	static UE_API void SetIsRecorded(ULevelSequence* InLevelSequence, bool bInIsRecorded);
+	/** Sets whether this is a subsequence. */
+	UFUNCTION(BlueprintCallable, Category = "Level Sequence")
+	static UE_API void SetIsSubSequence(ULevelSequence* InLevelSequence, bool bInIsSubSequence);
 	/** Sets the favorite rating. The favorite rating is like a star rating, usually 1-3 if it was rated. */
 	UFUNCTION(BlueprintCallable, Category = "Level Sequence")
 	static UE_API void SetFavoriteRating(ULevelSequence* InLevelSequence, int32 InFavoriteRating);
@@ -66,6 +90,12 @@ public:
 	/** Clear IsFlagged flag. */
 	UFUNCTION(BlueprintCallable, Category = "Level Sequence")
 	static UE_API void ClearIsFlagged(ULevelSequence* InLevelSequence);
+	/** Clear IsRecorded flag. */
+	UFUNCTION(BlueprintCallable, Category = "Level Sequence")
+	static UE_API void ClearIsRecorded(ULevelSequence* InLevelSequence);
+	/** Clear IsSubSequence flag. */
+	UFUNCTION(BlueprintCallable, Category = "Level Sequence")
+	static UE_API void ClearIsSubSequence(ULevelSequence* InLevelSequence);
 	/** Clears the favorite rating. */
 	UFUNCTION(BlueprintCallable, Category = "Level Sequence")
 	static UE_API void ClearFavoriteRating(ULevelSequence* InLevelSequence);
@@ -76,6 +106,12 @@ public:
 	/** @return The asset tag for the IsFlagged value. */
 	UFUNCTION(BlueprintPure, Category = "Level Sequence")
 	static UE_API FName GetIsFlaggedAssetTag();
+	/** @return The asset tag for the IsRecorded value. */
+	UFUNCTION(BlueprintPure, Category = "Level Sequence")
+	static UE_API FName GetIsRecordedAssetTag();
+	/** @return The asset tag for the IsSubSequence value. */
+	UFUNCTION(BlueprintPure, Category = "Level Sequence")
+	static UE_API FName GetIsSubSequenceAssetTag();
 	/** @return The asset tag for the FavoriteRating value. */
 	UFUNCTION(BlueprintPure, Category = "Level Sequence")
 	static UE_API FName GetFavoriteRatingAssetTag();
@@ -92,6 +128,18 @@ public:
 	 */
 	UFUNCTION(BlueprintPure, Category = "Level Sequence")
 	static UE_API bool GetIsFlaggedByAssetData(const FAssetData& InAssetData, bool& bOutIsFlagged);
+	/**
+	 * Gets whether this was recorded.
+	 * @return Whether the value was set by the user.
+	 */
+	UFUNCTION(BlueprintPure, Category = "Level Sequence")
+	static UE_API bool GetIsRecordedByAssetData(const FAssetData& InAssetData, bool& bOutIsRecorded);
+	/**
+	 * Gets whether this was recorded as a sub sequence.
+	 * @return Whether the value was set by the user.
+	 */
+	UFUNCTION(BlueprintPure, Category = "Level Sequence")
+	static UE_API bool GetIsSubSequenceByAssetData(const FAssetData& InAssetData, bool& bOutIsSubSequence);
 	/**
 	 * Gets the favorite rating. The favorite rating is like a star rating, usually 1-3 if it was rated.
 	 * @return Whether the value was set by the user.
@@ -107,6 +155,12 @@ public:
 	static UE_API bool HasIsFlaggedByAssetData(const FAssetData& InAssetData);
 	/** @return Whether the user has ever set a value for this property. */
 	UFUNCTION(BlueprintPure, Category = "Level Sequence")
+	static UE_API bool HasIsRecordedByAssetData(const FAssetData& InAssetData);
+	/** @return Whether the user has ever set a value for this property. */
+	UFUNCTION(BlueprintPure, Category = "Level Sequence")
+	static UE_API bool HasIsSubSequenceByAssetData(const FAssetData& InAssetData);
+	/** @return Whether the user has ever set a value for this property. */
+	UFUNCTION(BlueprintPure, Category = "Level Sequence")
 	static UE_API bool HasFavoriteRatingByAssetData(const FAssetData& InAssetData);
 };
 
diff --git a/Engine/Source/Runtime/MovieScene/Private/MetaData/MovieSceneShotMetaData.cpp b/Engine/Source/Runtime/MovieScene/Private/MetaData/MovieSceneShotMetaData.cpp
--- a/Engine/Source/Runtime/MovieScene/Private/MetaData/MovieSceneShotMetaData.cpp
+++ b/Engine/Source/Runtime/MovieScene/Private/MetaData/MovieSceneShotMetaData.cpp
@@ -29,6 +29,8 @@ static bool ParseAssetTag(const FAssetData& InAssetData, FName InTag, TTagValue&
 
 const FName UMovieSceneShotMetaData::AssetRegistryTag_bIsNoGood(TEXT("AssetRegistryTag_MovieScene_bIsNoGood"));
 const FName UMovieSceneShotMetaData::AssetRegistryTag_bIsFlagged(TEXT("AssetRegistryTag_MovieScene_bIsFlagged"));
+const FName UMovieSceneShotMetaData::AssetRegistryTag_bIsSubSequence(TEXT("AssetRegistryTag_IsSubSequence"));
+const FName UMovieSceneShotMetaData::AssetRegistryTag_bIsRecorded(TEXT("AssetRegistryTag_IsRecorded"));
 const FName UMovieSceneShotMetaData::AssetRegistryTag_FavoriteRating(TEXT("AssetRegistryTag_MovieScene_FavoriteRating"));
 
 bool UMovieSceneShotMetaData::GetIsNoGoodByAssetData(const FAssetData& InAssetData, bool& bOutNoGood)
@@ -41,6 +43,16 @@ bool UMovieSceneShotMetaData::GetIsFlaggedByAssetData(const FAssetData& InAssetD
 	return UE::MovieScene::ShotMetaDataDetail::ParseAssetTag(InAssetData, AssetRegistryTag_bIsFlagged, bOutIsFlagged);
 }
 
+bool UMovieSceneShotMetaData::GetIsRecordedByAssetData(const FAssetData& InAssetData, bool& bOutIsRecorded)
+{
+	return UE::MovieScene::ShotMetaDataDetail::ParseAssetTag(InAssetData, AssetRegistryTag_bIsRecorded, bOutIsRecorded);
+}
+
+bool UMovieSceneShotMetaData::GetIsSubSequenceByAssetData(const FAssetData& InAssetData, bool& bOutIsSubSequence)
+{
+	return UE::MovieScene::ShotMetaDataDetail::ParseAssetTag(InAssetData, AssetRegistryTag_bIsSubSequence, bOutIsSubSequence);
+}
+
 bool UMovieSceneShotMetaData::GetFavoriteRatingByAssetData(const FAssetData& InAssetData, int32& OutFavoriteRating)
 {
 	return UE::MovieScene::ShotMetaDataDetail::ParseAssetTag(InAssetData, AssetRegistryTag_FavoriteRating, OutFavoriteRating);
@@ -56,6 +68,14 @@ void UMovieSceneShotMetaData::ExtendAssetRegistryTags(FAssetRegistryTagsContext
 		AssetRegistryTag_bIsFlagged, bIsFlagged.IsSet() ? FString::FromInt(*bIsFlagged) : TEXT("Unset"),
 		FAssetRegistryTag::ETagType::TT_Numerical, FAssetRegistryTag::TD_None
 	));
+	Context.AddTag(FAssetRegistryTag(
+		AssetRegistryTag_bIsRecorded, bIsRecorded.IsSet() ? FString::FromInt(*bIsRecorded) : TEXT("Unset"),
+		FAssetRegistryTag::ETagType::TT_Numerical, FAssetRegistryTag::TD_None
+	));
+	Context.AddTag(FAssetRegistryTag(
+		AssetRegistryTag_bIsSubSequence, bIsSubSequence.IsSet() ? FString::FromInt(*bIsSubSequence) : TEXT("Unset"),
+		FAssetRegistryTag::ETagType::TT_Numerical, FAssetRegistryTag::TD_None
+	));
 	Context.AddTag(FAssetRegistryTag(
 		AssetRegistryTag_FavoriteRating, FavoriteRating.IsSet() ? FString::FromInt(*FavoriteRating) : TEXT("Unset"),
 		FAssetRegistryTag::ETagType::TT_Alphabetical, FAssetRegistryTag::TD_None
@@ -73,6 +93,14 @@ void UMovieSceneShotMetaData::ExtendAssetRegistryTagMetaData(TMap<FName, UObject
 		.SetDisplayName(LOCTEXT("IsFlagged.Label", "Is Flagged"))
 		.SetTooltip(LOCTEXT("IsFlagged.Description", "Whether this shot was highlighted"))
 	);
+	OutMetadata.Add(AssetRegistryTag_bIsRecorded, FAssetRegistryTagMetadata()
+		.SetDisplayName(LOCTEXT("IsRecorded.Label", "Is Recorded"))
+		.SetTooltip(LOCTEXT("IsRecorded.Description", "If this sequence was recorded."))
+	);
+	OutMetadata.Add(AssetRegistryTag_bIsSubSequence, FAssetRegistryTagMetadata()
+		.SetDisplayName(LOCTEXT("IsSubSequence.Label", "Is SubSequence"))
+		.SetTooltip(LOCTEXT("IsSubSequence.Description", "If this was recorded as a subsequence."))
+	);
 	OutMetadata.Add(AssetRegistryTag_FavoriteRating, FAssetRegistryTagMetadata()
 		.SetDisplayName(LOCTEXT("FavoriteRating.Label", "Rating"))
 		.SetTooltip(LOCTEXT("FavoriteRating.Description", "A star rating, usually ranging 1-3 and 0 meaning no rating, yet."))
diff --git a/Engine/Source/Runtime/MovieScene/Public/MetaData/MovieSceneShotMetaData.h b/Engine/Source/Runtime/MovieScene/Public/MetaData/MovieSceneShotMetaData.h
--- a/Engine/Source/Runtime/MovieScene/Public/MetaData/MovieSceneShotMetaData.h
+++ b/Engine/Source/Runtime/MovieScene/Public/MetaData/MovieSceneShotMetaData.h
@@ -29,25 +29,47 @@ public:
 	static UE_API const FName AssetRegistryTag_bIsNoGood;
 	/** The asset registry tag that contains whether this was flagged by a user */
 	static UE_API const FName AssetRegistryTag_bIsFlagged;
+	/** The asset registry tag that contains whether this was recorded. */
+	static UE_API const FName AssetRegistryTag_bIsRecorded;
+	/** The asset registry tag that contains whether this is a subsequence. */
+	static UE_API const FName AssetRegistryTag_bIsSubSequence;
 	/** The asset registry tag that contains its favorite status */
 	static UE_API const FName AssetRegistryTag_FavoriteRating;
 
-	static TArray<FName> GetAllTags() { return { AssetRegistryTag_bIsNoGood, AssetRegistryTag_bIsFlagged, AssetRegistryTag_FavoriteRating }; }
+	static TArray<FName> GetAllTags()
+	{
+		return
+		{
+			AssetRegistryTag_bIsNoGood,
+			AssetRegistryTag_bIsFlagged,
+			AssetRegistryTag_bIsRecorded,
+			AssetRegistryTag_bIsSubSequence,
+			AssetRegistryTag_FavoriteRating
+		};
+	}
 
 	const TOptional<bool>& GetIsNoGood() const { return bIsNoGood; }
 	const TOptional<bool>& GetIsFlagged() const { return bIsFlagged; }
+	const TOptional<bool>& GetIsRecorded() const { return bIsRecorded; }
+	const TOptional<bool>& GetIsSubSequence() const { return bIsSubSequence; }
 	const TOptional<int32>& GetFavoriteRating() const { return FavoriteRating; }
 
 	void SetIsNoGood(bool bInIsNoGood) { bIsNoGood = bInIsNoGood; }
 	void SetIsFlagged(bool bInIsFlagged) { bIsFlagged = bInIsFlagged; }
+	void SetIsRecorded(bool bInIsRecorded) { bIsRecorded = bInIsRecorded; }
+	void SetIsSubSequence(bool bInIsSubSequence) { bIsSubSequence = bInIsSubSequence; }
 	void SetFavoriteRating(int32 InFavoriteRating) { FavoriteRating = InFavoriteRating; }
 	
 	void ClearIsNoGood() { bIsNoGood.Reset(); }
 	void ClearIsFlagged() { bIsFlagged.Reset(); }
+	void ClearIsRecorded() { bIsRecorded.Reset(); }
+	void ClearIsSubSequence() { bIsSubSequence.Reset(); }
 	void ClearFavoriteRating() { FavoriteRating.Reset(); }
 
 	static UE_API bool GetIsNoGoodByAssetData(const FAssetData& InAssetData, bool& bOutNoGood);
 	static UE_API bool GetIsFlaggedByAssetData(const FAssetData& InAssetData, bool& bOutIsFlagged);
+	static UE_API bool GetIsRecordedByAssetData(const FAssetData& InAssetData, bool& bOutIsRecorded);
+	static UE_API bool GetIsSubSequenceByAssetData(const FAssetData& InAssetData, bool& bOutIsSubSequence);
 	static UE_API bool GetFavoriteRatingByAssetData(const FAssetData& InAssetData, int32& OutFavoriteRating); 
 
 	// Begin IMovieSceneMetaDataInterface Interface
@@ -67,6 +89,14 @@ private:
 	UPROPERTY()
 	TOptional<bool> bIsFlagged;
 
+	/** If this sequence was recorded. */
+	UPROPERTY()
+	TOptional<bool> bIsRecorded;
+
+	/** If this is a subsequence. */
+	UPROPERTY()
+	TOptional<bool> bIsSubSequence;
+	
 	/** The favorite rating is like a star rating, usually 1-3 if it was rated. */
 	UPROPERTY()
 	TOptional<int32> FavoriteRating;
