Expose CustomClock and MusicTime types to Blueprints to be able to do things in blueprints like cast a CustomClock to a custom "Musical Clock". This allows to get music time from a sequence via Blueprints

[RN] Provide functionality to get the Custom Clock from the Level Sequence Director in Blueprints

diff --git a/Engine/Plugins/MovieScene/SequencerScripting/Source/SequencerScripting/Private/ExtensionLibraries/MovieSceneSequenceExtensions.cpp b/Engine/Plugins/MovieScene/SequencerScripting/Source/SequencerScripting/Private/ExtensionLibraries/MovieSceneSequenceExtensions.cpp
--- a/Engine/Plugins/MovieScene/SequencerScripting/Source/SequencerScripting/Private/ExtensionLibraries/MovieSceneSequenceExtensions.cpp
+++ b/Engine/Plugins/MovieScene/SequencerScripting/Source/SequencerScripting/Private/ExtensionLibraries/MovieSceneSequenceExtensions.cpp
@@ -651,6 +651,23 @@ EUpdateClockSource UMovieSceneSequenceExtensions::GetClockSource(UMovieSceneSequ
 	return EUpdateClockSource::Tick;
 }
 
+UMovieSceneClock* UMovieSceneSequenceExtensions::GetCustomClock(UMovieSceneSequence* InSequence)
+{
+	if (!InSequence)
+	{
+		FFrame::KismetExecutionMessage(TEXT("Cannot call GetCustomClock on a null sequence"), ELogVerbosity::Error);
+		return nullptr;
+	}
+
+	UMovieScene* MovieScene = GetMovieScene(InSequence);
+	if (MovieScene)
+	{
+		return MovieScene->GetCustomClock();
+	}
+
+	return nullptr;
+}
+
 FMovieSceneBindingProxy UMovieSceneSequenceExtensions::FindBindingByName(UMovieSceneSequence* Sequence, FString Name)
 {
 	if (!Sequence)
diff --git a/Engine/Plugins/MovieScene/SequencerScripting/Source/SequencerScripting/Public/ExtensionLibraries/MovieSceneSequenceExtensions.h b/Engine/Plugins/MovieScene/SequencerScripting/Source/SequencerScripting/Public/ExtensionLibraries/MovieSceneSequenceExtensions.h
--- a/Engine/Plugins/MovieScene/SequencerScripting/Source/SequencerScripting/Public/ExtensionLibraries/MovieSceneSequenceExtensions.h
+++ b/Engine/Plugins/MovieScene/SequencerScripting/Source/SequencerScripting/Public/ExtensionLibraries/MovieSceneSequenceExtensions.h
@@ -342,6 +342,15 @@ public:
 	UFUNCTION(BlueprintPure, Category = "Sequencer|Sequence", meta = (ScriptMethod))
 	static EUpdateClockSource GetClockSource(UMovieSceneSequence* InSequence);
 
+	/**
+	 * Get the custom clock for this sequence
+	 * 
+	 * @param Sequence The Sequence within which to get the custom clock
+	 * @return The MovieSceneClock for this sequence
+	 */
+	UFUNCTION(BlueprintPure, Category = "Sequencer|Sequence", meta = (ScriptMethod))
+	static UMovieSceneClock* GetCustomClock(UMovieSceneSequence* InSequence);
+
 	/**
 	 * Attempt to locate a binding in this sequence by its name
 	 *
diff --git a/Engine/Source/Runtime/CoreUObject/Public/UObject/NoExportTypes.h b/Engine/Source/Runtime/CoreUObject/Public/UObject/NoExportTypes.h
--- a/Engine/Source/Runtime/CoreUObject/Public/UObject/NoExportTypes.h
+++ b/Engine/Source/Runtime/CoreUObject/Public/UObject/NoExportTypes.h
@@ -2242,7 +2242,7 @@ struct FTimespan
  * Represents a time by a context-free musical bar/beat. 
  * @note The full C++ class is located here: Engine\Source\Runtime\Core\Public\Misc\MusicalTime.h
  */
-USTRUCT(noexport)
+USTRUCT(noexport, BlueprintType)
 struct FMusicalTime
 {
 	/** 0-based Bar Index */
diff --git a/Engine/Source/Runtime/LevelSequence/Private/LevelSequenceDirector.cpp b/Engine/Source/Runtime/LevelSequence/Private/LevelSequenceDirector.cpp
--- a/Engine/Source/Runtime/LevelSequence/Private/LevelSequenceDirector.cpp
+++ b/Engine/Source/Runtime/LevelSequence/Private/LevelSequenceDirector.cpp
@@ -42,6 +42,51 @@ const UE::MovieScene::FSequenceInstance* ULevelSequenceDirector::FindSequenceIns
 	return nullptr;
 }
 
+UMovieSceneClock* ULevelSequenceDirector::GetRootSequenceCustomClock() const
+{
+	using namespace UE::MovieScene;
+
+	if (const FSequenceInstance* Instance = FindSequenceInstance())
+	{
+		TSharedPtr<const FSharedPlaybackState> SharedPlaybackState = Instance->GetSharedPlaybackState();
+		UMovieSceneSequence* RootSequence = SharedPlaybackState->GetRootSequence();
+		if (RootSequence)
+		{
+			if (UMovieScene* MovieScene = RootSequence->GetMovieScene())
+			{
+				return MovieScene->GetCustomClock();
+			}
+		}
+	}
+	return nullptr;
+}
+
+UMovieSceneClock* ULevelSequenceDirector::GetSequenceCustomClock() const
+{
+	using namespace UE::MovieScene;
+
+	if (const FSequenceInstance* Instance = FindSequenceInstance())
+	{
+		TSharedPtr<const FSharedPlaybackState> SharedPlaybackState = Instance->GetSharedPlaybackState();
+		UMovieSceneSequence* Sequence = SharedPlaybackState->GetRootSequence();
+
+		const FMovieSceneSequenceID ActualSubSequenceID(SubSequenceID);
+		if (ActualSubSequenceID != MovieSceneSequenceID::Root)
+		{
+			Sequence = SharedPlaybackState->GetHierarchy()->FindSubSequence(ActualSubSequenceID);
+		}
+
+		if (Sequence)
+		{
+			if (UMovieScene* MovieScene = Sequence->GetMovieScene())
+			{
+				return MovieScene->GetCustomClock();
+			}
+		}
+	}
+	return nullptr;
+}
+
 FQualifiedFrameTime ULevelSequenceDirector::GetRootSequenceTime() const
 {
 	using namespace UE::MovieScene;
diff --git a/Engine/Source/Runtime/LevelSequence/Public/LevelSequenceDirector.h b/Engine/Source/Runtime/LevelSequence/Public/LevelSequenceDirector.h
--- a/Engine/Source/Runtime/LevelSequence/Public/LevelSequenceDirector.h
+++ b/Engine/Source/Runtime/LevelSequence/Public/LevelSequenceDirector.h
@@ -3,6 +3,7 @@
 #pragma once
 
 #include "CoreTypes.h"
+#include "MovieSceneClock.h"
 #include "Engine/Blueprint.h"
 #include "Engine/BlueprintGeneratedClass.h"
 #include "EntitySystem/MovieSceneSequenceInstance.h"
@@ -25,6 +26,20 @@ public:
 	/** Called when this director is created */
 	UFUNCTION(BlueprintImplementableEvent, Category="Sequencer")
 	LEVELSEQUENCE_API void OnCreated();
+
+	/**
+	 * 
+	 * @return The custom clock of the outermost (root) sequence
+	 */
+	UFUNCTION(BlueprintCallable, Category = "Sequencer|Director")
+	LEVELSEQUENCE_API UMovieSceneClock* GetRootSequenceCustomClock() const;
+
+	/**
+	 * 
+	 * @return The custom clock of this director's sequence
+	 */
+	UFUNCTION(BlueprintCallable, Category = "Sequencer|Director")
+	LEVELSEQUENCE_API UMovieSceneClock* GetSequenceCustomClock() const;
 	
 	/**
 	 * Get the current time for the outermost (root) sequence
diff --git a/Engine/Source/Runtime/TimeManagement/Private/MusicalTimeFunctionLibrary.cpp b/Engine/Source/Runtime/TimeManagement/Private/MusicalTimeFunctionLibrary.cpp
new file mode 100644
--- /dev/null
+++ b/Engine/Source/Runtime/TimeManagement/Private/MusicalTimeFunctionLibrary.cpp
@@ -0,0 +1,25 @@
+﻿// Copyright Epic Games, Inc. All Rights Reserved.
+
+#include "MusicalTimeFunctionLibrary.h"
+#include "Misc/MusicalTime.h"
+
+bool UMusicalTimeFunctionLibrary::IsValid(const FMusicalTime& InMusicalTime)
+{
+	return InMusicalTime.IsValid();
+}
+
+float UMusicalTimeFunctionLibrary::FractionalBeatsInBar(const FMusicalTime& InMusicalTime)
+{
+	return InMusicalTime.FractionalBeatInBar();
+}
+	
+float UMusicalTimeFunctionLibrary::FractionalBars(const FMusicalTime& InMusicalTime)
+{
+	return InMusicalTime.FractionalBar();
+}
+	
+void UMusicalTimeFunctionLibrary::BarAndBeat(const FMusicalTime& InMusicalTime, int& OutBar, float& OutFractionalBeatInBar)
+{
+	OutBar = InMusicalTime.Bar;
+	OutFractionalBeatInBar = InMusicalTime.FractionalBeatInBar();
+}
\ No newline at end of file
diff --git a/Engine/Source/Runtime/TimeManagement/Public/MusicalTimeFunctionLibrary.h b/Engine/Source/Runtime/TimeManagement/Public/MusicalTimeFunctionLibrary.h
new file mode 100644
--- /dev/null
+++ b/Engine/Source/Runtime/TimeManagement/Public/MusicalTimeFunctionLibrary.h
@@ -0,0 +1,37 @@
+﻿// Copyright Epic Games, Inc. All Rights Reserved.
+
+#pragma once
+
+#include "Kismet/BlueprintFunctionLibrary.h"
+#include "Misc/MusicalTime.h"
+
+#include "MusicalTimeFunctionLibrary.generated.h"
+
+
+class UObject;
+struct FFrame;
+struct FMusicalTime;
+
+
+/**
+ * Expose FMusicalTime methods to blueprints
+ */
+UCLASS(meta = (BlueprintThreadSafe, ScriptName = "MusicalTimeFunctionLibrary"), MinimalAPI)
+class UMusicalTimeFunctionLibrary : public UBlueprintFunctionLibrary
+{
+	GENERATED_BODY()
+
+public:
+
+	UFUNCTION(BlueprintPure, meta = (DisplayName = "Is Valid Musical Time"), Category = "Utilities|Musical Time")
+	static bool IsValid(const FMusicalTime& InMusicalTime);
+
+	UFUNCTION(BlueprintPure, meta = (DisplayName = "Fractional Beat In Bar"), Category = "Utilities|Musical Time")
+	static float FractionalBeatsInBar(const FMusicalTime& InMusicalTime);
+	
+	UFUNCTION(BlueprintPure, meta = (DisplayName = "Fractional Bar"), Category = "Utilities|Musical Time")
+	static float FractionalBars(const FMusicalTime& InMusicalTime);
+	
+	UFUNCTION(BlueprintPure, meta = (DisplayName = "Bar & Beat"), Category = "Utilities|Musical Time")
+	static void BarAndBeat(const FMusicalTime& InMusicalTime, int& Bar, float& Beat);
+};
\ No newline at end of file
