Sequencer: Use AddUniqueDynamic instead of AddDynamic

In some cases, the delegate appears to have been already added. This is a speculative fix.

#rnx

diff --git a/Engine/Source/Runtime/LevelSequence/Private/LevelSequenceActor.cpp b/Engine/Source/Runtime/LevelSequence/Private/LevelSequenceActor.cpp
--- a/Engine/Source/Runtime/LevelSequence/Private/LevelSequenceActor.cpp
+++ b/Engine/Source/Runtime/LevelSequence/Private/LevelSequenceActor.cpp
@@ -226,9 +226,9 @@ void ALevelSequenceActor::BeginPlay()
 	if (GetSequencePlayer())
 	{
 		AddReplicatedSubObject(GetSequencePlayer());
-		GetSequencePlayer()->OnPlay.AddDynamic(this, &ALevelSequenceActor::ShowBurnin);
-		GetSequencePlayer()->OnPlayReverse.AddDynamic(this, &ALevelSequenceActor::ShowBurnin);
-		GetSequencePlayer()->OnStop.AddDynamic(this, &ALevelSequenceActor::HideBurnin);
+		GetSequencePlayer()->OnPlay.AddUniqueDynamic(this, &ALevelSequenceActor::ShowBurnin);
+		GetSequencePlayer()->OnPlayReverse.AddUniqueDynamic(this, &ALevelSequenceActor::ShowBurnin);
+		GetSequencePlayer()->OnStop.AddUniqueDynamic(this, &ALevelSequenceActor::HideBurnin);
 	}
 
 	if (PlaybackSettings.bAutoPlay)
