Sequencer: Mark the movie scene as changed when the transform origin on the Level Sequence Actor changes

This fixes a bug in the editor when you edit the transform origin so that the sequence will re-evaluate

#Synced-CL 48857117

diff --git a/Engine/Source/Runtime/LevelSequence/Private/DefaultLevelSequenceInstanceData.cpp b/Engine/Source/Runtime/LevelSequence/Private/DefaultLevelSequenceInstanceData.cpp
new file mode 100644
--- /dev/null
+++ b/Engine/Source/Runtime/LevelSequence/Private/DefaultLevelSequenceInstanceData.cpp
@@ -0,0 +1,17 @@
+// Copyright Epic Games, Inc. All Rights Reserved.
+
+#include "DefaultLevelSequenceInstanceData.h"
+
+#if WITH_EDITOR
+void UDefaultLevelSequenceInstanceData::PostEditChangeProperty(FPropertyChangedEvent& PropertyChangedEvent)
+{
+	FName PropertyName = (PropertyChangedEvent.Property != nullptr) ? PropertyChangedEvent.Property->GetFName() : NAME_None;
+
+	if (PropertyName == GET_MEMBER_NAME_CHECKED(UDefaultLevelSequenceInstanceData, TransformOriginActor) || PropertyName == GET_MEMBER_NAME_CHECKED(UDefaultLevelSequenceInstanceData, TransformOrigin))
+	{
+		TransformOriginChangedEvent.Broadcast();
+	}
+
+	Super::PostEditChangeProperty(PropertyChangedEvent);
+}
+#endif // WITH_EDITOR
\ No newline at end of file
diff --git a/Engine/Source/Runtime/LevelSequence/Private/LevelSequenceActor.cpp b/Engine/Source/Runtime/LevelSequence/Private/LevelSequenceActor.cpp
--- a/Engine/Source/Runtime/LevelSequence/Private/LevelSequenceActor.cpp
+++ b/Engine/Source/Runtime/LevelSequence/Private/LevelSequenceActor.cpp
@@ -80,6 +80,13 @@ ALevelSequenceActor::ALevelSequenceActor(const FObjectInitializer& Init)
 	BurnInOptions = Init.CreateDefaultSubobject<ULevelSequenceBurnInOptions>(this, "BurnInOptions");
 	DefaultInstanceData = Init.CreateDefaultSubobject<UDefaultLevelSequenceInstanceData>(this, "InstanceData");
 
+#if WITH_EDITOR
+	if (UDefaultLevelSequenceInstanceData* LevelSequenceInstanceData = Cast<UDefaultLevelSequenceInstanceData>(DefaultInstanceData))
+	{
+		LevelSequenceInstanceData->TransformOriginChangedEvent.AddUObject(this, &ALevelSequenceActor::TransformOriginChanged);
+	}
+#endif // WITH_EDITOR
+
 	// SequencePlayer must be a default sub object for it to be replicated correctly
 PRAGMA_DISABLE_DEPRECATION_WARNINGS // make SequencePlayer protected and remove this for 5.6
 	SequencePlayer = Init.CreateDefaultSubobject<ULevelSequencePlayer>(this, "AnimationPlayer");
@@ -321,6 +328,34 @@ void ALevelSequenceActor::PostLoad()
 #endif
 }
 
+#if WITH_EDITOR
+
+void ALevelSequenceActor::PostEditChangeProperty(FPropertyChangedEvent& PropertyChangedEvent)
+{
+	FName PropertyName = (PropertyChangedEvent.Property != nullptr) ? PropertyChangedEvent.Property->GetFName() : NAME_None;
+
+	if (PropertyName == GET_MEMBER_NAME_CHECKED(ALevelSequenceActor, bOverrideInstanceData))
+	{
+		TransformOriginChanged();
+	}
+
+	Super::PostEditChangeProperty(PropertyChangedEvent);
+}
+
+void ALevelSequenceActor::TransformOriginChanged()
+{
+	if (LevelSequenceAsset)
+	{
+		if (UMovieScene* MovieScene = LevelSequenceAsset->GetMovieScene())
+		{
+			MovieScene->MarkAsChanged();
+		}
+	}
+}
+
+#endif // #if WITH_EDITOR
+
+
 #if WITH_EDITORONLY_DATA
 void ALevelSequenceActor::DeclareConstructClasses(TArray<FTopLevelAssetPath>& OutConstructClasses, const UClass* SpecificSubclass)
 {
diff --git a/Engine/Source/Runtime/LevelSequence/Public/DefaultLevelSequenceInstanceData.h b/Engine/Source/Runtime/LevelSequence/Public/DefaultLevelSequenceInstanceData.h
--- a/Engine/Source/Runtime/LevelSequence/Public/DefaultLevelSequenceInstanceData.h
+++ b/Engine/Source/Runtime/LevelSequence/Public/DefaultLevelSequenceInstanceData.h
@@ -35,4 +35,12 @@ public:
 	/** Specifies a transform that offsets all absolute transform sections in this sequence. Scale is ignored. Not applied to Relative or Additive sections. */
 	UPROPERTY(EditAnywhere, BlueprintReadWrite, Category="General")
 	FTransform TransformOrigin;
+
+#if WITH_EDITOR
+	DECLARE_MULTICAST_DELEGATE(FTransformOriginChanged);
+	FTransformOriginChanged TransformOriginChangedEvent;
+
+	//~ Begin UObject interface
+	virtual void PostEditChangeProperty(FPropertyChangedEvent& PropertyChangedEvent) override;
+#endif // WITH_EDITOR
 };
diff --git a/Engine/Source/Runtime/LevelSequence/Public/LevelSequenceActor.h b/Engine/Source/Runtime/LevelSequence/Public/LevelSequenceActor.h
--- a/Engine/Source/Runtime/LevelSequence/Public/LevelSequenceActor.h
+++ b/Engine/Source/Runtime/LevelSequence/Public/LevelSequenceActor.h
@@ -264,6 +264,13 @@ protected:
 	//~ Begin UObject interface
 	LEVELSEQUENCE_API virtual void PostInitProperties() override;
 	LEVELSEQUENCE_API virtual void PostLoad() override;
+
+#if WITH_EDITOR
+	LEVELSEQUENCE_API virtual void PostEditChangeProperty(FPropertyChangedEvent& PropertyChangedEvent) override;
+
+	void TransformOriginChanged();
+#endif //WITH_EDITOR
+
 public:
 #if WITH_EDITORONLY_DATA
 	static LEVELSEQUENCE_API void DeclareConstructClasses(TArray<FTopLevelAssetPath>& OutConstructClasses, const UClass* SpecificSubclass);
