Fix for when an animation instance is a post process anim instance that is being evaluated we would fail to verify if it has changed because we don't check that it is post process. This marks the sequence as dirty and triggers the spawning phase to happen every frame. Now we validate against post process animbps as well.

#Synced-CL 52308368

diff --git a/Engine/Source/Runtime/LevelSequence/Private/LevelSequence.cpp b/Engine/Source/Runtime/LevelSequence/Private/LevelSequence.cpp
--- a/Engine/Source/Runtime/LevelSequence/Private/LevelSequence.cpp
+++ b/Engine/Source/Runtime/LevelSequence/Private/LevelSequence.cpp
@@ -717,10 +717,25 @@ void ULevelSequence::GatherExpiredObjects(const FMovieSceneObjectCache& InObject
 		
 		if (Reference.Locator.GetLastFragmentTypeHandle() == FAnimInstanceLocatorFragment::FragmentType)
 		{
+			// Get anim instance fragment to determine type.
+			EAnimInstanceLocatorFragmentType LocatorType = EAnimInstanceLocatorFragmentType::AnimInstance;
+			if (const FAnimInstanceLocatorFragment* AnimFragment = Reference.Locator.GetLastFragment()->GetPayloadAs(FAnimInstanceLocatorFragment::FragmentType))
+			{
+				LocatorType = AnimFragment->Type;
+			}
+			
 			for (TWeakObjectPtr<> WeakObject : InObjectCache.IterateBoundObjects(Reference.ID))
 			{
 				UAnimInstance* AnimInstance = Cast<UAnimInstance>(WeakObject.Get());
-				if (!AnimInstance || !AnimInstance->GetOwningComponent() || AnimInstance->GetOwningComponent()->GetAnimInstance() != AnimInstance)
+				UAnimInstance* AnimInstanceToTest = nullptr;
+				if (AnimInstance && AnimInstance->GetOwningComponent())
+				{
+					//Check if we're a regular or post-process anim instance and if the owning component's anim instance has changed.
+					AnimInstanceToTest = LocatorType == EAnimInstanceLocatorFragmentType::PostProcessAnimInstance ?
+						AnimInstance->GetOwningComponent()->GetPostProcessInstance() :
+						AnimInstance->GetOwningComponent()->GetAnimInstance();
+				}
+				if (!AnimInstance || !AnimInstance->GetOwningComponent() || AnimInstanceToTest != AnimInstance)
 				{
 					OutInvalidIDs.Add(Reference.ID);
 				}
